hololens platform test file.

#rnx

diff --git a/Engine/Platforms/Hololens/test.txt b/Engine/Platforms/Hololens/test.txt
new file mode 100644
--- /dev/null
+++ b/Engine/Platforms/Hololens/test.txt
@@ -0,0 +1 @@
+this is just a test
