move HoloLens build & automation code to platform extension folder.

#rnx

diff --git a/Engine/Platforms/Hololens/Source/Programs/AutomationTool/HoloLens.Automation.csproj b/Engine/Platforms/Hololens/Source/Programs/AutomationTool/HoloLens.Automation.csproj
new file mode 100644
--- /dev/null
+++ b/Engine/Platforms/Hololens/Source/Programs/AutomationTool/HoloLens.Automation.csproj
@@ -0,0 +1,57 @@
+﻿<Project Sdk="Microsoft.NET.Sdk">
+  <Import Project="..\..\..\..\..\..\Engine\Source\Programs\Shared\UnrealEngine.csproj.props" />
+  
+  <PropertyGroup>
+    <TargetFramework>net6.0-windows10.0.19041.0</TargetFramework>
+    <Configuration Condition=" '$(Configuration)' == '' ">Development</Configuration>
+    <OutputType>Library</OutputType>
+    <AppendTargetFrameworkToOutputPath>false</AppendTargetFrameworkToOutputPath>
+    <GenerateAssemblyInfo>false</GenerateAssemblyInfo>
+    <GenerateTargetFrameworkAttribute>false</GenerateTargetFrameworkAttribute>
+    <Configurations>Debug;Release;Development</Configurations>
+    <RootNamespace>HoloLens.Automation</RootNamespace>
+    <AssemblyName>HoloLens.Automation</AssemblyName>
+    <WarningsNotAsErrors>612,618</WarningsNotAsErrors>
+    <OutputPath>..\..\..\Binaries\DotNET\AutomationTool\AutomationScripts\Platforms\HoloLens\</OutputPath>
+    <TreatWarningsAsErrors>true</TreatWarningsAsErrors>
+    <DebugType>pdbonly</DebugType>
+  </PropertyGroup>
+  <PropertyGroup Condition=" '$(Configuration)|$(Platform)' == 'Development|AnyCPU' ">
+    <DefineConstants>$(DefineConstants);TRACE</DefineConstants>
+    <Optimize>true</Optimize>
+  </PropertyGroup>
+  <PropertyGroup Condition=" '$(Configuration)|$(Platform)' == 'Release|AnyCPU' ">
+    <DefineConstants>$(DefineConstants);TRACE</DefineConstants>
+    <Optimize>true</Optimize>
+  </PropertyGroup>
+  <PropertyGroup Condition=" '$(Configuration)|$(Platform)' == 'Debug|AnyCPU' ">
+    <DefineConstants>$(DefineConstants);DEBUG;TRACE</DefineConstants>
+    <DebugSymbols>true</DebugSymbols>
+    <DebugType>full</DebugType>
+  </PropertyGroup>
+  <ItemGroup>
+    <Reference Include="WindowsDevicePortalWrapper">
+      <HintPath>..\..\..\..\..\..\Engine\Binaries\ThirdParty\WindowsDevicePortalWrapper\WindowsDevicePortalWrapper.dll</HintPath>
+    </Reference>
+  </ItemGroup>
+  <PropertyGroup>
+    <UseWindowsForms>true</UseWindowsForms>
+  </PropertyGroup>
+
+  <ItemGroup>
+    <Compile Include="..\..\..\..\..\..\Engine\Source\Programs\Shared\MetaData.cs">
+      <Link>Properties\MetaData.cs</Link>
+    </Compile>
+  </ItemGroup>
+  <ItemGroup>
+    <ProjectReference Include="..\..\..\..\..\..\Engine\Source\Programs\Shared\EpicGames.Core\EpicGames.Core.csproj" PrivateAssets="All"><Private>false</Private></ProjectReference>
+    <ProjectReference Include="..\..\..\..\..\..\Engine\Source\Programs\UnrealBuildTool\UnrealBuildTool.csproj" PrivateAssets="All"><Private>false</Private></ProjectReference>
+    <ProjectReference Include="..\..\..\..\..\..\Engine\Source\Programs\AutomationTool\AutomationUtils\AutomationUtils.Automation.csproj" PrivateAssets="All"><Private>false</Private></ProjectReference>
+    <ProjectReference Include="..\..\..\..\..\..\Engine\Source\Programs\AutomationTool\Scripts\AutomationScripts.Automation.csproj" PrivateAssets="All"><Private>false</Private></ProjectReference>
+  </ItemGroup>
+  <ItemGroup>
+    <PackageReference Include="Microsoft.CSharp" Version="4.7.0" />
+    <PackageReference Include="System.Data.DataSetExtensions" Version="4.5.0" />
+    <PackageReference Include="System.Runtime.WindowsRuntime" Version="4.7.0" />
+  </ItemGroup>
+</Project>
diff --git a/Engine/Platforms/Hololens/Source/Programs/AutomationTool/HoloLensPlatform.Automation.cs b/Engine/Platforms/Hololens/Source/Programs/AutomationTool/HoloLensPlatform.Automation.cs
new file mode 100644
--- /dev/null
+++ b/Engine/Platforms/Hololens/Source/Programs/AutomationTool/HoloLensPlatform.Automation.cs
@@ -0,0 +1,1827 @@
+// Copyright Epic Games, Inc. All Rights Reserved.
+
+using System;
+using System.Collections.Generic;
+using System.Linq;
+using System.Text;
+using System.IO;
+using System.Threading.Tasks;
+using AutomationTool;
+using UnrealBuildTool;
+using System.Diagnostics;
+using EpicGames.Core;
+using System.Text.RegularExpressions;
+
+using Windows.Management.Deployment;
+using Windows.Foundation;
+using System.Threading;
+using System.Security.Cryptography.X509Certificates;
+using Microsoft.Tools.WindowsDevicePortal;
+using UnrealBuildBase;
+using AutomationScripts;
+using System.Runtime.Versioning;
+
+namespace HoloLens.Automation
+{
+	class HoloLensLauncherCreatedProcess : IProcessResult
+	{
+		Process ExternallyLaunchedProcess;
+		System.Threading.Thread LogTailThread;
+		bool AllowSpew;
+		LogEventType SpewVerbosity;
+
+		public HoloLensLauncherCreatedProcess(Process InExternallyLaunchedProcess, string LogFile, bool InAllowSpew, LogEventType InSpewVerbosity)
+		{
+			ExternallyLaunchedProcess = InExternallyLaunchedProcess;
+			AllowSpew = InAllowSpew;
+			SpewVerbosity = InSpewVerbosity;
+
+			// Can't redirect stdout, so tail the log file
+			if (AllowSpew)
+			{
+				LogTailThread = new System.Threading.Thread(() => TailLogFile(LogFile));
+				LogTailThread.Start();
+			}
+			else
+			{
+				LogTailThread = null;
+			}
+		}
+
+		~HoloLensLauncherCreatedProcess()
+		{
+			if (ExternallyLaunchedProcess != null)
+			{
+				ExternallyLaunchedProcess.Dispose();
+			}
+		}
+
+		public int ExitCode
+		{
+			get
+			{
+				// Can't access the exit code of a process we didn't start
+				return 0;
+			}
+
+			set
+			{
+				throw new NotImplementedException();
+			}
+		}
+
+		public bool HasExited
+		{
+			get
+			{
+				// Avoid potential access of the exit code.
+				return ExternallyLaunchedProcess != null && ExternallyLaunchedProcess.HasExited;
+			}
+		}
+
+		public string Output
+		{
+			get
+			{
+				return string.Empty;
+			}
+		}
+
+		public Process ProcessObject
+		{
+			get
+			{
+				return ExternallyLaunchedProcess;
+			}
+		}
+
+		public void DisposeProcess()
+		{
+			ExternallyLaunchedProcess.Dispose();
+			ExternallyLaunchedProcess = null;
+		}
+
+		public string GetProcessName()
+		{
+			return ExternallyLaunchedProcess.ProcessName;
+		}
+
+		public void OnProcessExited()
+		{
+			ProcessManager.RemoveProcess(this);
+			if (LogTailThread != null)
+			{
+				LogTailThread.Join();
+				LogTailThread = null;
+			}
+		}
+
+		public void StdErr(object sender, DataReceivedEventArgs e)
+		{
+			throw new NotImplementedException();
+		}
+
+		public void StdOut(object sender, DataReceivedEventArgs e)
+		{
+			throw new NotImplementedException();
+		}
+
+		public void StopProcess(bool KillDescendants = true)
+		{
+			string ProcessNameForLogging = GetProcessName();
+			try
+			{
+				Process ProcessToKill = ExternallyLaunchedProcess;
+				ExternallyLaunchedProcess = null;
+				if (!ProcessToKill.CloseMainWindow())
+				{
+					CommandUtils.LogWarning("{0} did not respond to close request.  Killing...", ProcessNameForLogging);
+					ProcessToKill.Kill();
+					ProcessToKill.WaitForExit(60000);
+				}
+				if (!ProcessToKill.HasExited)
+				{
+					CommandUtils.LogLog("Process {0} failed to exit.", ProcessNameForLogging);
+				}
+				else
+				{
+					CommandUtils.LogLog("Process {0} successfully exited.", ProcessNameForLogging);
+					OnProcessExited();
+				}
+				ProcessToKill.Close();
+			}
+			catch (Exception Ex)
+			{
+				CommandUtils.LogWarning("Exception while trying to kill process {0}:", ProcessNameForLogging);
+				CommandUtils.LogWarning(LogUtils.FormatException(Ex));
+			}
+		}
+
+		public void WaitForExit()
+		{
+			ExternallyLaunchedProcess.WaitForExit();
+			if (LogTailThread != null)
+			{
+				LogTailThread.Join();
+				LogTailThread = null;
+			}
+		}
+
+		private void TailLogFile(string LogFilePath)
+		{
+			string LogName = GetProcessName();
+			while (!HasExited && !File.Exists(LogFilePath))
+			{
+				System.Threading.Thread.Sleep(1000);
+			}
+
+			if (File.Exists(LogFilePath))
+			{
+				using (StreamReader LogReader = new StreamReader(new FileStream(LogFilePath, FileMode.Open, FileAccess.Read, FileShare.ReadWrite)))
+				{
+					while (!HasExited)
+					{
+						System.Threading.Thread.Sleep(5000);
+						while (!LogReader.EndOfStream)
+						{
+							string LogLine = LogReader.ReadLine();
+							Log.WriteLine(SpewVerbosity, "{0} : {1}", LogName, LogLine);
+						}
+					}
+				}
+			}
+		}
+	}
+
+	class HoloLensDevicePortalCreatedProcess : IProcessResult
+	{
+		object StateLock;
+		Microsoft.Tools.WindowsDevicePortal.DevicePortal Portal;
+		uint ProcId;
+		string PackageName;
+		string FriendlyName;
+
+		bool ProcessHasExited;
+
+		public HoloLensDevicePortalCreatedProcess(Microsoft.Tools.WindowsDevicePortal.DevicePortal InPortal, string InPackageName, string InFriendlyName)
+		{
+			Portal = InPortal;
+			PackageName = InPackageName;
+			FriendlyName = InFriendlyName;
+			StateLock = new object();
+			//MonitorThread = new System.Threading.Thread(()=>(MonitorProcess))
+
+			// No ETW on Xbox One, so can't collect trace that way
+			if (Portal.Platform != Microsoft.Tools.WindowsDevicePortal.DevicePortal.DevicePortalPlatforms.XboxOne)
+			{
+				Guid MicrosoftWindowsDiagnoticsLoggingChannelId = new Guid(0x4bd2826e, 0x54a1, 0x4ba9, 0xbf, 0x63, 0x92, 0xb7, 0x3e, 0xa1, 0xac, 0x4a);
+				Portal.ToggleEtwProviderAsync(MicrosoftWindowsDiagnoticsLoggingChannelId).Wait();
+				Portal.RealtimeEventsMessageReceived += Portal_RealtimeEventsMessageReceived;
+				Portal.StartListeningForEtwEventsAsync().Wait();
+			}
+		}
+
+		public void SetProcId(uint InProcId)
+		{
+			ProcId = InProcId;
+			Portal.RunningProcessesMessageReceived += Portal_RunningProcessesMessageReceived;
+			Portal.StartListeningForRunningProcessesAsync().Wait();
+		}
+
+		private void Portal_RealtimeEventsMessageReceived(Microsoft.Tools.WindowsDevicePortal.DevicePortal sender, Microsoft.Tools.WindowsDevicePortal.WebSocketMessageReceivedEventArgs<Microsoft.Tools.WindowsDevicePortal.DevicePortal.EtwEvents> args)
+		{
+			foreach (var EtwEvent in args.Message.Events)
+			{
+				if (EtwEvent.ContainsKey("ProviderName") && EtwEvent.ContainsKey("StringMessage"))
+				{
+					if (EtwEvent["ProviderName"] == FriendlyName)
+					{
+						Log.WriteLine(GetLogVerbosityFromEventLevel(EtwEvent.Level), "{0} : {1}", FriendlyName, EtwEvent["StringMessage"].Trim('\"'));
+					}
+				}
+			}
+		}
+
+		private LogEventType GetLogVerbosityFromEventLevel(uint EtwLevel)
+		{
+			switch (EtwLevel)
+			{
+				case 4:
+					return LogEventType.Console;
+				case 3:
+					return LogEventType.Warning;
+				case 2:
+					return LogEventType.Error;
+				case 1:
+					return LogEventType.Fatal;
+				default:
+					return LogEventType.Verbose;
+			}
+		}
+
+		private void Portal_RunningProcessesMessageReceived(Microsoft.Tools.WindowsDevicePortal.DevicePortal sender, Microsoft.Tools.WindowsDevicePortal.WebSocketMessageReceivedEventArgs<Microsoft.Tools.WindowsDevicePortal.DevicePortal.RunningProcesses> args)
+		{
+			foreach (var P in args.Message.Processes)
+			{
+				if (P.ProcessId == ProcId)
+				{
+					return;
+				}
+			}
+			ProcessHasExited = true;
+		}
+
+		public int ExitCode
+		{
+			get
+			{
+				return 0;
+			}
+
+			set
+			{
+				throw new NotImplementedException();
+			}
+		}
+
+		public bool HasExited
+		{
+			get
+			{
+				return ProcessHasExited;
+			}
+		}
+
+		public string Output
+		{
+			get
+			{
+				return string.Empty;
+			}
+		}
+
+		public Process ProcessObject
+		{
+			get
+			{
+				return null;
+			}
+		}
+
+		public void DisposeProcess()
+		{
+		}
+
+		public string GetProcessName()
+		{
+			return FriendlyName;
+		}
+
+		public void OnProcessExited()
+		{
+			Portal.StopListeningForRunningProcessesAsync().Wait();
+			if (Portal.Platform != Microsoft.Tools.WindowsDevicePortal.DevicePortal.DevicePortalPlatforms.XboxOne)
+			{
+				Portal.StopListeningForEtwEventsAsync().Wait();
+			}
+		}
+
+		public void StdErr(object sender, DataReceivedEventArgs e)
+		{
+			throw new NotImplementedException();
+		}
+
+		public void StdOut(object sender, DataReceivedEventArgs e)
+		{
+			throw new NotImplementedException();
+		}
+
+		public void StopProcess(bool KillDescendants = true)
+		{
+			Portal.TerminateApplicationAsync(PackageName).Wait();
+		}
+
+		public void WaitForExit()
+		{
+			while (!ProcessHasExited)
+			{
+				System.Threading.Thread.Sleep(1000);
+			}
+		}
+	}
+
+	[SupportedOSPlatform("windows10.0.10240.0")]
+	public class HoloLensPlatform : Platform
+    {
+		private WindowsArchitecture[] ActualArchitectures = { };
+
+		private FileReference MakeAppXPath;
+		private FileReference PDBCopyPath;
+		private FileReference SignToolPath;
+		private FileReference MakeCertPath;
+		private FileReference Pvk2PfxPath;
+		private string Windows10SDKVersion;
+		private string Extension = ".appx";
+
+
+		public HoloLensPlatform()
+			: base(UnrealTargetPlatform.HoloLens)
+		{
+
+		}
+
+		public override void PreBuildAgenda(UnrealBuild Build, UnrealBuild.BuildAgenda Agenda, ProjectParams Params)
+		{
+			if(ActualArchitectures.Length == 0)
+			{
+				throw new AutomationException(ExitCode.Error_Arguments, "No target architecture selected on \'Platforms/HoloLens/OS Info\' page. Please select at last one.");
+			}
+
+			foreach (var BuildConfig in Params.ClientConfigsToBuild)
+			{
+				foreach(var target in Params.ClientCookedTargets)
+				{
+					var ReceiptFileName = TargetReceipt.GetDefaultPath(Params.RawProjectPath.Directory, target, UnrealTargetPlatform.HoloLens, BuildConfig, "Multi");
+					if(File.Exists(ReceiptFileName.FullName))
+					{
+						FileReference.Delete(ReceiptFileName);
+					}
+
+					foreach (var Arch in ActualArchitectures)
+					{
+						Agenda.Targets.Add(new UnrealBuild.BuildTarget()
+						{
+							TargetName = target,
+							Platform = UnrealTargetPlatform.HoloLens,
+							Config = BuildConfig,
+							UprojectPath = Params.CodeBasedUprojectPath,
+							UBTArgs = " -remoteini=\"" + Params.RawProjectPath.Directory.FullName + "\" -Architecture=" + WindowsExports.GetArchitectureSubpath(Arch),
+						});
+					}
+				}
+			}
+		}
+
+		public override bool CanBeCompiled() //block auto compilation
+		{
+			return false;
+		}
+
+
+		public override void PlatformSetupParams(ref ProjectParams ProjParams)
+		{
+			base.PlatformSetupParams(ref ProjParams);
+
+			if (ProjParams.Deploy && !ProjParams.Package)
+			{
+				foreach (string DeviceAddress in ProjParams.DeviceNames)
+				{
+					if (!IsLocalDevice(DeviceAddress))
+					{
+						LogWarning("Project will be packaged to support deployment to remote HoloLens device {0}.", DeviceAddress);
+						ProjParams.Package = true;
+						break;
+					}
+				}
+			}
+
+			ConfigHierarchy PlatformEngineConfig = null;
+			if (!ProjParams.EngineConfigs.TryGetValue(PlatformType, out PlatformEngineConfig))
+			{
+				throw new AutomationException(ExitCode.Error_Arguments, "No configuration on \'Platforms/HoloLens/OS Info\' page. Please create one.");
+			}
+
+			PlatformEngineConfig.GetString("/Script/HoloLensPlatformEditor.HoloLensTargetSettings", "Windows10SDKVersion", out Windows10SDKVersion);
+
+			List<string> ThumbprintsFromConfig = new List<string>();
+			if (PlatformEngineConfig.GetArray("/Script/HoloLensPlatformEditor.HoloLensTargetSettings", "AcceptThumbprints", out ThumbprintsFromConfig))
+			{
+				AcceptThumbprints.AddRange(ThumbprintsFromConfig);
+			}
+
+
+			if (ProjParams.Run)
+			{
+				var ArchList = new List<WindowsArchitecture>();
+				foreach (string DeviceAddress in ProjParams.DeviceNames)
+				{
+					//We have to choose architecture of the device to run
+					WindowsArchitecture Arch = WindowsArchitecture.x64;
+					if (!IsLocalDevice(DeviceAddress))
+					{
+						Arch = RemoteDeviceArchitecture(DeviceAddress, ProjParams);
+					}
+
+					ArchList.Add(Arch);
+
+					LogInformation(String.Format("Project will be compiled for the architecture {0} of the HoloLens device {1}.", WindowsExports.GetArchitectureSubpath(Arch), DeviceAddress));
+				}
+
+				ActualArchitectures = ArchList.Distinct().ToArray();
+			}
+			else
+			{
+				var ArchList = new List<WindowsArchitecture>();
+				bool bBuildForEmulation = false;
+				bool bBuildForDevice = false;
+                
+				if (PlatformEngineConfig.GetBool("/Script/HoloLensPlatformEditor.HoloLensTargetSettings", "bBuildForEmulation", out bBuildForEmulation))
+				{
+					if (bBuildForEmulation)
+					{
+						ArchList.Add(WindowsArchitecture.x64);
+					}
+				}
+
+				if (PlatformEngineConfig.GetBool("/Script/HoloLensPlatformEditor.HoloLensTargetSettings", "bBuildForDevice", out bBuildForDevice))
+				{
+					// If both are unchecked, then we build for device
+					if (bBuildForDevice || (!bBuildForEmulation && !bBuildForDevice))
+					{
+						ArchList.Add(WindowsArchitecture.ARM64);
+					}
+				}
+
+				ActualArchitectures = ArchList.ToArray();
+			}
+
+			string ArchString = ActualArchitectures.Length == 1 && ProjParams.Run && !ProjParams.Package ? WindowsExports.GetArchitectureSubpath(ActualArchitectures[0]) : "Multi";
+			ProjParams.SpecifiedArchitecture = ArchString;
+			ProjParams.ConfigOverrideParams.Add(String.Format("Architecture={0}", ArchString));
+
+			FindInstruments();
+			GenerateSigningCertificate(ProjParams, PlatformEngineConfig);
+		}
+
+		private void FindInstruments()
+		{
+			if(string.IsNullOrEmpty(Windows10SDKVersion))
+			{
+				Windows10SDKVersion = "Latest";
+			}
+
+			if(!HoloLensExports.InitWindowsSdkToolPath(Windows10SDKVersion))
+			{
+				throw new AutomationException(ExitCode.Error_Arguments, "Wrong WinSDK toolchain selected on \'Platforms/HoloLens/Toolchain\' page. Please check.");
+			}
+
+			// VS puts MSBuild stuff (where PDBCopy lives) under the Visual Studio Installation directory
+			DirectoryReference MSBuildInstallDir = new DirectoryReference(Path.Combine(WindowsExports.GetMSBuildToolPath(), "..", "..", ".."));
+
+			DirectoryReference SDKFolder;
+			Version SDKVersion;
+
+			if(WindowsExports.TryGetWindowsSdkDir(Windows10SDKVersion, out SDKVersion, out SDKFolder))
+			{
+				DirectoryReference WindowsSdkBinDir = DirectoryReference.Combine(SDKFolder, "bin", SDKVersion.ToString(), Environment.Is64BitProcess ? "x64" : "x86");
+				if (!DirectoryReference.Exists(WindowsSdkBinDir))
+				{
+					throw new AutomationException(ExitCode.Error_Arguments, "WinSDK toolchain selected on \'Platforms/HoloLens/Toolchain\' page isn't exist anymore. Please select new one.");
+				}
+				if (PDBCopyPath == null || !FileReference.Exists(PDBCopyPath))
+				{
+					PDBCopyPath = FileReference.Combine(SDKFolder, "Debuggers", Environment.Is64BitProcess ? "x64" : "x86", "PDBCopy.exe");
+				}
+			}
+
+
+			MakeAppXPath = HoloLensExports.GetWindowsSdkToolPath("makeappx.exe");
+			SignToolPath = HoloLensExports.GetWindowsSdkToolPath("signtool.exe");
+			MakeCertPath = HoloLensExports.GetWindowsSdkToolPath("makecert.exe");
+			Pvk2PfxPath = HoloLensExports.GetWindowsSdkToolPath("pvk2pfx.exe");
+
+			IEnumerable<DirectoryReference> VSInstallDirs;
+			if (PDBCopyPath == null || !FileReference.Exists(PDBCopyPath))
+			{
+				if (null != (VSInstallDirs = WindowsExports.TryGetVSInstallDirs(WindowsCompiler.VisualStudio2019)))
+				{
+					PDBCopyPath = FileReference.Combine(VSInstallDirs.First(), "MSBuild", "Microsoft", "VisualStudio", "v16.0", "AppxPackage", "PDBCopy.exe");
+				}
+				else if (null != (VSInstallDirs = WindowsExports.TryGetVSInstallDirs(WindowsCompiler.VisualStudio2022)))
+				{
+					PDBCopyPath = FileReference.Combine(VSInstallDirs.First(), "MSBuild", "Microsoft", "VisualStudio", "v17.0", "AppxPackage", "PDBCopy.exe");
+				}
+			}
+
+			// Earlier versions use a separate MSBuild install location
+			if (PDBCopyPath == null || !FileReference.Exists(PDBCopyPath))
+			{
+				PDBCopyPath = FileReference.Combine(MSBuildInstallDir, "Microsoft", "VisualStudio", "v14.0", "AppxPackage", "PDBCopy.exe");
+			}
+
+			if (PDBCopyPath == null || !FileReference.Exists(PDBCopyPath))
+			{
+				PDBCopyPath = FileReference.Combine(MSBuildInstallDir, "Microsoft", "VisualStudio", "v12.0", "AppxPackage", "PDBCopy.exe");
+			}
+
+			if (!FileReference.Exists(PDBCopyPath))
+			{
+				PDBCopyPath = null;
+			}
+
+		}
+
+		public override void Deploy(ProjectParams Params, DeploymentContext SC)
+		{
+			foreach (string DeviceAddress in Params.DeviceNames)
+			{
+				if (IsLocalDevice(DeviceAddress))
+				{
+					// Special case - we can use PackageManager to allow loose apps, plus we need to
+					// apply the loopback exemption in case of cook-on-the-fly
+					DeployToLocalDevice(Params, SC);
+				}
+				else
+				{
+					DeployToRemoteDevice(DeviceAddress, Params, SC);
+				}
+			}
+		}
+
+		public override void GetFilesToDeployOrStage(ProjectParams Params, DeploymentContext SC)
+		{
+			if(Params.Run && !Params.Package)
+			{
+				// Stage all the build products
+				DirectoryReference ProjectBinariesFolder = Params.GetProjectBinariesPathForPlatform(PlatformType);
+				HoloLensExports DeployExports = new HoloLensExports();
+				foreach (StageTarget Target in SC.StageTargets)
+				{
+					string ArchString = Target.Receipt.Architecture;
+					SC.StageBuildProductsFromReceipt(Target.Receipt, Target.RequireFilesExist, Params.bTreatNonShippingBinariesAsDebugFiles);
+					DeployExports.AddWinMDReferencesFromReceipt(Target.Receipt, Params.RawProjectPath.Directory, SC.LocalRoot.FullName, Windows10SDKVersion);
+
+					// Stage HoloLens-specific assets (tile, splash, etc.)
+					DirectoryReference assetsPath = new DirectoryReference(Path.Combine(ProjectBinariesFolder.FullName, ArchString, "Resources"));
+					StagedDirectoryReference stagedAssetPath = new StagedDirectoryReference("Resources");
+					SC.StageFiles(StagedFileType.NonUFS, assetsPath, "*.png", StageFilesSearch.AllDirectories, stagedAssetPath);
+
+
+					SC.StageFile(StagedFileType.NonUFS, new FileReference(Path.Combine(ProjectBinariesFolder.FullName, String.Format("AppxManifest_{0}.xml", ArchString))),
+						new StagedFileReference("AppxManifest.xml"));
+					SC.StageFile(StagedFileType.NonUFS, new FileReference(Path.Combine(ProjectBinariesFolder.FullName, String.Format("resources_{0}.pri", ArchString))),
+						new StagedFileReference("resources.pri"));
+
+					FileReference SourceNetworkManifestPath = new FileReference(Path.Combine(ProjectBinariesFolder.FullName, String.Format("NetworkManifest_{0}.xml", ArchString)));
+					if (FileReference.Exists(SourceNetworkManifestPath))
+					{
+						SC.StageFile(StagedFileType.NonUFS, SourceNetworkManifestPath, new StagedFileReference("NetworkManifest.xml"));
+					}
+
+					TargetRules Rules = Params.ProjectTargets.Find(x => x.Rules.Type == TargetType.Game).Rules;
+					bool UseDebugCrt = Params.ClientConfigsToBuild.Contains(UnrealTargetConfiguration.Debug) && Rules.bDebugBuildsActuallyUseDebugCRT;
+					foreach (var RT in GetPathToVCLibsPackages(UseDebugCrt, Rules.WindowsPlatform.Compiler))
+					{
+						SC.StageFile(StagedFileType.NonUFS, new FileReference(RT), new StagedFileReference(Path.GetFileName(RT)));
+					}
+				}
+			}
+		}
+
+		public override string GetCookPlatform(bool bDedicatedServer, bool bIsClientOnly)
+		{
+		    string Return = PlatformType.ToString();
+            return bIsClientOnly ? Return + "Client" : Return;
+		}
+
+		static void FillMapfile(DirectoryReference DirectoryName, string SearchPath, string Mask, StringBuilder AppXRecipeBuiltFiles)
+		{
+			string StageDir = DirectoryName.FullName;
+			if (!StageDir.EndsWith("\\") && !StageDir.EndsWith("/"))
+			{
+				StageDir += Path.DirectorySeparatorChar;
+			}
+			Uri StageDirUri = new Uri(StageDir);
+			foreach (var pf in Directory.GetFiles(SearchPath, Mask, SearchOption.AllDirectories))
+			{
+				var FileUri = new Uri(pf);
+				var relativeUri = StageDirUri.MakeRelativeUri(FileUri);
+				var relativePath = Uri.UnescapeDataString(relativeUri.ToString());
+				var newPath = relativePath.Replace('/', Path.DirectorySeparatorChar);
+
+				AppXRecipeBuiltFiles.AppendLine(String.Format("\"{0}\"\t\"{1}\"", pf, newPath));
+			}
+		}
+
+		static void FillMapfile(DirectoryReference DirectoryName, FileReference ManifestFile, StringBuilder AppXRecipeBuiltFiles)
+		{
+			if (FileReference.Exists(ManifestFile))
+			{
+				string[] Lines = FileReference.ReadAllLines(ManifestFile);
+				foreach (string Line in Lines)
+				{
+					string[] Pair = Line.Split('\t');
+					if (Pair.Length > 1)
+					{
+						AppXRecipeBuiltFiles.AppendLine(String.Format("\"{0}\"\t\"{1}\"", FileReference.Combine(DirectoryName, Pair[0]), Pair[0]));
+					}
+				}
+			}
+		}
+
+		static void FillMapfile(DirectoryReference DirectoryName, string SearchPath, string Mask, Dictionary<string, FileReference> Dict)
+		{
+			string StageDir = DirectoryName.FullName;
+			if (!StageDir.EndsWith("\\") && !StageDir.EndsWith("/"))
+			{
+				StageDir += Path.DirectorySeparatorChar;
+			}
+			Uri StageDirUri = new Uri(StageDir);
+			foreach (var pf in Directory.GetFiles(SearchPath, Mask, SearchOption.AllDirectories))
+			{
+				var FileUri = new Uri(pf);
+				var relativeUri = StageDirUri.MakeRelativeUri(FileUri);
+				var relativePath = Uri.UnescapeDataString(relativeUri.ToString());
+				var newPath = relativePath.Replace('/', Path.DirectorySeparatorChar);
+
+				if (!Dict.ContainsKey(newPath))
+				{
+					Dict[newPath] = new FileReference(pf);
+				}
+			}
+		}
+
+		static void FillMapfile(DirectoryReference DirectoryName, FileReference ManifestFile, Dictionary<string, FileReference> Dict)
+		{
+			if (FileReference.Exists(ManifestFile))
+			{
+				string[] Lines = FileReference.ReadAllLines(ManifestFile);
+				foreach (string Line in Lines)
+				{
+					string[] Pair = Line.Split('\t');
+					if (Pair.Length > 1)
+					{
+						if(!Dict.ContainsKey(Pair[0]))
+						{
+							Dict[Pair[0]] = FileReference.Combine(DirectoryName, Pair[0]);
+						}
+					}
+				}
+			}
+		}
+
+		private List<string> GetContainerFileExtensions()
+		{
+			return new List<string> { "*.pak", "*.ucas", "*.utoc" };
+		}
+
+		private bool IsUsingContainers(ProjectParams Params)
+		{
+			return (Params.UsePak(this) || (Params.IoStore && !Params.SkipIoStore));
+		}
+
+		private void PackagePakFiles(ProjectParams Params, DeploymentContext SC, string OutputNameBase)
+		{
+			string IntermediateDirectory = Path.Combine(SC.ProjectRoot.FullName, "Intermediate", "Deploy", "neutral");
+			var ListResources = new HoloLensManifestGenerator().CreateAssetsManifest(SC.StageTargetPlatform.PlatformType, SC.StageDirectory.FullName, IntermediateDirectory, SC.RawProjectPath, SC.ProjectRoot.FullName);
+
+			string OutputName = OutputNameBase + "_pak";
+
+			var AppXRecipeBuiltFiles = new StringBuilder();
+
+			string MapFilename = Path.Combine(SC.StageDirectory.FullName, OutputName + ".pkgmap");
+			AppXRecipeBuiltFiles.AppendLine("[ResourceMetadata]");
+			AppXRecipeBuiltFiles.AppendLine("\"ResourceId\"\t\"UnrealAssets\"");
+			AppXRecipeBuiltFiles.AppendLine("");
+			AppXRecipeBuiltFiles.AppendLine(@"[Files]");
+
+			string OutputAppX = Path.Combine(SC.StageDirectory.FullName, OutputName + Extension);
+
+			if(IsUsingContainers(Params))
+			{
+				foreach (string ContainerExtension in GetContainerFileExtensions())
+				{
+					FillMapfile(SC.StageDirectory, SC.StageDirectory.FullName, ContainerExtension, AppXRecipeBuiltFiles);
+				}
+			}
+			else
+			{
+				FillMapfile(SC.StageDirectory, FileReference.Combine(SC.StageDirectory, SC.GetUFSDeployedManifestFileName(null)), AppXRecipeBuiltFiles);
+			}
+
+			FillMapfile(SC.StageDirectory, FileReference.Combine(SC.StageDirectory, SC.GetNonUFSDeployedManifestFileName(null)), AppXRecipeBuiltFiles);
+
+			{
+				DirectoryReference ResourceFolder = SC.StageDirectory;
+				foreach(var ResourcePath in ListResources)
+				{
+					var ResourceFile = new FileReference(ResourcePath);
+					if(ResourceFile.IsUnderDirectory(ResourceFolder))
+					{
+						AppXRecipeBuiltFiles.AppendLine(String.Format("\"{0}\"\t\"{1}\"", ResourceFile.FullName, ResourceFile.MakeRelativeTo(ResourceFolder)));
+					}
+				}
+			}
+
+			//AppXRecipeBuiltFiles.AppendLine(String.Format("\"{0}\"\t\"{1}\"", Path.Combine(SC.StageDirectory.FullName, "AppxManifest_assets.xml"), "AppxManifest.xml"));
+
+			File.WriteAllText(MapFilename, AppXRecipeBuiltFiles.ToString(), Encoding.UTF8);
+
+			string MakeAppXCommandLine = String.Format(@"pack /r /o /f ""{0}"" /p ""{1}"" /m ""{2}""", MapFilename, OutputAppX, Path.Combine(SC.StageDirectory.FullName, "AppxManifest_assets.xml"));
+			RunAndLog(CmdEnv, MakeAppXPath.FullName, MakeAppXCommandLine, null, 0, null, ERunOptions.None);
+			SignPackage(Params, SC, OutputAppX);
+		}
+
+		private void UpdateCodePackagesWithData(ProjectParams Params, DeploymentContext SC, string OutputNameBase)
+		{
+			foreach (StageTarget Target in SC.StageTargets)
+			{
+				foreach (BuildProduct Product in Target.Receipt.BuildProducts)
+				{
+					if (Product.Type != BuildProductType.MapFile)
+					{
+						continue;
+					}
+
+					string MapFilename = Product.Path.FullName; 
+
+					var AppXRecipeBuiltFiles = new StringBuilder();
+
+					string OutputAppX = Path.Combine(SC.StageDirectory.FullName, Product.Path.GetFileNameWithoutExtension() + Extension);
+
+					if(IsUsingContainers(Params))
+					{
+						foreach (string ContainerExtension in GetContainerFileExtensions())
+						{
+							FillMapfile(SC.StageDirectory, SC.StageDirectory.FullName, ContainerExtension, AppXRecipeBuiltFiles);
+						}
+					}
+					else
+					{
+						FillMapfile(SC.StageDirectory, FileReference.Combine(SC.StageDirectory, SC.GetUFSDeployedManifestFileName(null)), AppXRecipeBuiltFiles);
+					}
+
+					FillMapfile(SC.StageDirectory, FileReference.Combine(SC.StageDirectory, SC.GetNonUFSDeployedManifestFileName(null)), AppXRecipeBuiltFiles);
+
+					File.AppendAllText(MapFilename, AppXRecipeBuiltFiles.ToString(), Encoding.UTF8);
+
+					string MakeAppXCommandLine = String.Format(@"pack /o /f ""{0}"" /p ""{1}""", MapFilename, OutputAppX);
+					RunAndLog(CmdEnv, MakeAppXPath.FullName, MakeAppXCommandLine, null, 0, null, ERunOptions.None);
+
+					SignPackage(Params, SC, OutputAppX);
+				}
+			}
+		}
+
+		private void MakeSingleApp(ProjectParams Params, DeploymentContext SC, string OutputNameBase)
+		{
+			foreach (StageTarget Target in SC.StageTargets)
+			{
+				string IntermediateDirectory = Path.Combine(SC.ProjectRoot.FullName, "Intermediate", "Deploy", Params.SpecifiedArchitecture);
+				var Receipt = Target.Receipt;
+				string OutputName = String.Format("{0}_{1}_{2}_{3}", Receipt.TargetName, Receipt.Platform, Receipt.Configuration, Receipt.Architecture);
+
+				string MapFilename = Path.Combine(IntermediateDirectory, OutputName + "_full.pkgmap");
+
+				Dictionary<string, FileReference> Dict = new Dictionary<string, FileReference>();
+
+
+				{
+					//parse old mapfile
+					string OldMapFilename = Path.Combine(IntermediateDirectory, OutputName + ".pkgmap");
+					string[] lines = File.ReadAllLines(OldMapFilename, Encoding.UTF8);
+					foreach (var line in lines)
+					{
+						if (line.Length == 0 || line[0] == '[')
+						{
+							continue;
+						}
+
+						string[] files = line.Split('\t');
+
+						if(files.Length != 2)
+						{
+							continue;
+						}
+
+						files[0] = files[0].Trim('\"');
+						files[1] = files[1].Trim('\"').Replace('\\', '/');
+
+						if (Dict.ContainsKey(files[1]))
+						{
+							continue;
+						}
+
+						Dict[files[1]] = new FileReference(files[0]);
+					}
+				}
+
+				string OutputAppX = Path.Combine(SC.StageDirectory.FullName, OutputName + Extension);
+
+				if (IsUsingContainers(Params) && !Params.Run)
+				{
+					foreach (string ContainerExtension in GetContainerFileExtensions())
+					{
+						FillMapfile(SC.StageDirectory, SC.StageDirectory.FullName, ContainerExtension, Dict);
+					}
+				}
+				else
+				{
+					FillMapfile(SC.StageDirectory, FileReference.Combine(SC.StageDirectory, SC.GetUFSDeployedManifestFileName(null)), Dict);
+				}
+
+				FillMapfile(SC.StageDirectory, FileReference.Combine(SC.StageDirectory, SC.GetNonUFSDeployedManifestFileName(null)), Dict);
+
+
+				var AppXRecipeBuiltFiles = new StringBuilder();
+
+				AppXRecipeBuiltFiles.AppendLine("[Files]");
+				foreach (var P in Dict)
+				{
+					AppXRecipeBuiltFiles.AppendLine(String.Format("\"{0}\"\t\"{1}\"", P.Value, P.Key));
+				}
+
+				File.WriteAllText(MapFilename, AppXRecipeBuiltFiles.ToString(), Encoding.UTF8);
+
+				string MakeAppXCommandLine = String.Format("pack /o /f \"{0}\" /p \"{1}\"", MapFilename, OutputAppX);
+				RunAndLog(CmdEnv, MakeAppXPath.FullName, MakeAppXCommandLine, null, 0, null, ERunOptions.None);
+
+				SignPackage(Params, SC, OutputAppX);
+			}
+		}
+		
+		void MakeAppInstaller(ProjectParams Params, DeploymentContext SC, string OutputNameBase, string autoUpdateURL, int hoursBetweenUpdates)
+		{
+			// Validate url
+			string url = autoUpdateURL.Trim();
+			{
+				if (!url.EndsWith("/"))
+				{
+					url += "/";
+				}
+
+				if (!url.ToLower().StartsWith("http:") && !url.ToLower().StartsWith("https:") && !url.ToLower().StartsWith("www."))
+				{
+					if (!url.StartsWith("file:"))
+					{
+						url = "file:" + url;
+					}
+				}
+
+				url = url.Replace('\\', '/');
+			}
+
+			string version;
+			string packageName; 
+			string appxBundleName = OutputNameBase + ".appxbundle";
+			string publisherName;
+
+			ConfigHierarchy PlatformGameConfig = null;
+			Params.GameConfigs.TryGetValue(PlatformType, out PlatformGameConfig);
+			PlatformGameConfig.GetString("/Script/EngineSettings.GeneralProjectSettings", "ProjectName", out packageName);
+			PlatformGameConfig.GetString("/Script/EngineSettings.GeneralProjectSettings", "CompanyDistinguishedName", out publisherName);
+			PlatformGameConfig.GetString("/Script/EngineSettings.GeneralProjectSettings", "ProjectVersion", out version);
+
+			// Package name from project settings may be unique words, but installer and manifest need to strip whitespace to validate.
+			packageName = packageName.Replace(" ", String.Empty);
+
+			string VCLibsName = "Microsoft.VCLibs.140.00";
+			string VCLibsPackage = "Microsoft.VCLibs.arm64.14.00.appx";
+			string VCLibsInfo = "Publisher=\"CN=Microsoft Corporation, O=Microsoft Corporation, L=Redmond, S=Washington, C=US\" Version=\"14.0.0.0\" ProcessorArchitecture=\"arm64\"";
+			TargetRules Rules = Params.ProjectTargets.Find(x => x.Rules.Type == TargetType.Game).Rules;
+			bool UseDebugCrt = Params.ClientConfigsToBuild.Contains(UnrealTargetConfiguration.Debug) && Rules.bDebugBuildsActuallyUseDebugCRT;
+			if (UseDebugCrt)
+			{
+				VCLibsName += ".Debug";
+				VCLibsPackage = "Microsoft.VCLibs.arm64.Debug.14.00.appx";
+			}
+
+			// Ideally we would use an XmlWriter, but the installer uses an xmlns attribute which the XmlWriter fails to validate.
+			var builder = new StringBuilder();
+
+			string AppInstallerFilename = Path.Combine(SC.StageDirectory.FullName, OutputNameBase + ".appinstaller");
+			builder.AppendLine("<?xml version=\"1.0\" encoding=\"utf-8\"?>");
+			builder.AppendLine("<AppInstaller Uri=" + "\"" + url + OutputNameBase + ".appinstaller" + "\"" + " Version=" + "\"" + version + "\"" + " xmlns=\"http://schemas.microsoft.com/appx/appinstaller/2017/2\">");
+			builder.AppendLine("  <MainBundle Name=" + "\"" + packageName + "\"" + " Version=" + "\"" + version + "\"" + " Publisher=" + "\"" + publisherName + "\"" + " Uri=" + "\"" + url + appxBundleName + "\"" + " />");
+			builder.AppendLine("  <Dependencies>");
+			builder.AppendLine("    <Package Name=" + "\"" + VCLibsName + "\" " + VCLibsInfo + " Uri=" + "\"" + url + VCLibsPackage + "\"" + " />");
+			builder.AppendLine("  </Dependencies>");
+			builder.AppendLine("  <UpdateSettings>");
+			builder.AppendLine("    <OnLaunch HoursBetweenUpdateChecks=" + "\"" + hoursBetweenUpdates + "\"" + " />");
+			builder.AppendLine("  </UpdateSettings>");
+			builder.AppendLine("</AppInstaller>");
+
+			File.WriteAllText(AppInstallerFilename, builder.ToString(), Encoding.UTF8);
+		}
+
+		void MakeBundle(ProjectParams Params, DeploymentContext SC, string OutputNameBase, bool SeparateAssetPackaging)
+		{
+			string OutputName = OutputNameBase;
+
+			var AppXRecipeBuiltFiles = new StringBuilder();
+
+			string MapFilename = Path.Combine(SC.StageDirectory.FullName, OutputName + "_bundle.pkgmap");
+			AppXRecipeBuiltFiles.AppendLine(@"[Files]");
+
+			string OutputAppX = Path.Combine(SC.StageDirectory.FullName, OutputName + Extension + "bundle");
+			if(SeparateAssetPackaging)
+			{
+				foreach (StageTarget Target in SC.StageTargets)
+				{
+					foreach (BuildProduct Product in Target.Receipt.BuildProducts)
+					{
+						if (Product.Type == BuildProductType.Package)
+						{
+							AppXRecipeBuiltFiles.AppendLine(String.Format("\"{0}\"\t\"{1}\"", Product.Path.FullName, Path.GetFileName(Product.Path.FullName)));
+						}
+					}
+				}
+
+				AppXRecipeBuiltFiles.AppendLine(String.Format("\"{0}\"\t\"{1}\"", Path.Combine(SC.StageDirectory.FullName, OutputName + "_pak" + Extension), OutputName + "_pak" + Extension));//assets from pak file
+			}
+			else
+			{
+				FillMapfile(SC.StageDirectory, SC.StageDirectory.FullName, "*" + Extension, AppXRecipeBuiltFiles);
+			}
+
+			File.WriteAllText(MapFilename, AppXRecipeBuiltFiles.ToString(), Encoding.UTF8);
+
+			ConfigHierarchy PlatformGameConfig = null;
+			string projectVersion = "0.0.0.0";
+			if (Params.GameConfigs.TryGetValue(PlatformType, out PlatformGameConfig))
+			{
+				PlatformGameConfig.GetString("/Script/EngineSettings.GeneralProjectSettings", "ProjectVersion", out projectVersion);
+			}
+
+			string MakeAppXCommandLine = string.Format(@"bundle /o /f ""{0}"" /p ""{1}"" /bv {2}", MapFilename, OutputAppX, projectVersion);
+			RunAndLog(CmdEnv, MakeAppXPath.FullName, MakeAppXCommandLine, null, 0, null, ERunOptions.None);
+			SignPackage(Params, SC, OutputName + Extension + "bundle", true);
+		}
+
+		private void CopyVCLibs(ProjectParams Params, DeploymentContext SC)
+		{
+			TargetRules Rules = Params.ProjectTargets.Find(x => x.Rules.Type == TargetType.Game)?.Rules;
+
+			bool UseDebugCrt = false;
+			WindowsCompiler compiler = WindowsCompiler.VisualStudio2019;
+
+			//TODO: Why is this null?
+			if (Rules != null)
+			{
+				UseDebugCrt = Params.ClientConfigsToBuild.Contains(UnrealTargetConfiguration.Debug) && Rules.bDebugBuildsActuallyUseDebugCRT;
+				compiler = Rules.WindowsPlatform.Compiler;
+			}
+
+			foreach (string vcLib in GetPathToVCLibsPackages(UseDebugCrt, compiler))
+			{
+				CopyFile(vcLib, Path.Combine(SC.StageDirectory.FullName, Path.GetFileName(vcLib)));
+			}
+		}
+
+		private void GenerateSigningCertificate(ProjectParams Params, ConfigHierarchy PlatformEngineConfig)
+		{
+			string SigningCertificate = @"Build\HoloLens\SigningCertificate.pfx";
+			string SigningCertificatePath = Path.Combine(Params.RawProjectPath.Directory.FullName, SigningCertificate);
+			if (!File.Exists(SigningCertificatePath))
+			{
+				if (!IsBuildMachine && !Params.Unattended)
+				{
+					LogError("Certificate is required.  Please go to Project Settings > HoloLens > Create Signing Certificate");
+				}
+			}
+		}
+
+		private void SignPackage(ProjectParams Params, DeploymentContext SC, string OutputName, bool GenerateCer = false)
+		{
+			string OutputNameBase = Path.GetFileNameWithoutExtension(OutputName);
+			string OutputAppX = Path.Combine(SC.StageDirectory.FullName, OutputName);
+			string SigningCertificate = @"Build\HoloLens\SigningCertificate.pfx";
+
+			if (GenerateCer)
+			{
+				// Emit a .cer file adjacent to the appx so it can be installed to enable packaged deployment
+				System.Security.Cryptography.X509Certificates.X509Certificate2 ActualCert = new System.Security.Cryptography.X509Certificates.X509Certificate2(Path.Combine(SC.ProjectRoot.FullName, SigningCertificate));
+				File.WriteAllText(Path.Combine(SC.StageDirectory.FullName, OutputNameBase + ".cer"), Convert.ToBase64String(ActualCert.Export(System.Security.Cryptography.X509Certificates.X509ContentType.Cert)));
+			}
+
+			string CertFile = Path.Combine(SC.ProjectRoot.FullName, SigningCertificate);
+			if(File.Exists(CertFile))
+			{
+				string SignToolCommandLine = string.Format(@"sign /a /f ""{0}"" /fd SHA256 ""{1}""", CertFile, OutputAppX);
+				RunAndLog(CmdEnv, SignToolPath.FullName, SignToolCommandLine, null, 0, null, ERunOptions.None);
+			}
+		}
+
+		public override void Package(ProjectParams Params, DeploymentContext SC, int WorkingCL)
+		{
+			//GenerateDLCManifestIfNecessary(Params, SC);
+			string OutputNameBase = Params.HasDLCName ? Params.DLCFile.GetFileNameWithoutExtension() : Params.ShortProjectName;
+			string OutputAppX = Path.Combine(SC.StageDirectory.FullName, OutputNameBase + Extension);
+			{
+				OutputAppX += "bundle";
+				bool SeparateAssetPackaging = true;
+				bool bStartInVR = false;
+
+				//auto update
+				bool bShouldCreateAppInstaller = false;
+				string autoUpdateURL = string.Empty;
+				int hoursBetweenUpdates = 0;
+				
+				ConfigHierarchy PlatformEngineConfig = null;
+				if (Params.EngineConfigs.TryGetValue(PlatformType, out PlatformEngineConfig))
+				{
+					//PlatformEngineConfig.GetBool("/Script/HoloLensPlatformEditor.HoloLensTargetSettings", "bUseAssetPackage", out SeparateAssetPackaging);
+					
+					// Get auto update vars
+					PlatformEngineConfig.GetBool("/Script/HoloLensPlatformEditor.HoloLensTargetSettings", "bShouldCreateAppInstaller", out bShouldCreateAppInstaller);
+					PlatformEngineConfig.GetString("/Script/HoloLensPlatformEditor.HoloLensTargetSettings", "AppInstallerInstallationURL", out autoUpdateURL);
+					PlatformEngineConfig.GetInt32("/Script/HoloLensPlatformEditor.HoloLensTargetSettings", "HoursBetweenUpdateChecks", out hoursBetweenUpdates);
+				}
+
+				ConfigHierarchy PlatformGameConfig = null;
+				if (Params.GameConfigs.TryGetValue(PlatformType, out PlatformGameConfig))
+				{
+					PlatformGameConfig.GetBool("/Script/EngineSettings.GeneralProjectSettings", "bStartInVR", out bStartInVR);
+					Log.TraceInformation("bStartInVR = {0}", bStartInVR.ToString());
+				}
+
+				if (bStartInVR)
+				{
+					bool updateCommandLine = false;
+					if (string.IsNullOrEmpty(Params.StageCommandline))
+					{
+						Params.StageCommandline = "-vr";
+						updateCommandLine = true;
+					}
+					else if (!Params.StageCommandline.Contains("-vr"))
+					{
+						Params.StageCommandline += " -vr";
+						updateCommandLine = true;
+					}
+
+					if (updateCommandLine)
+					{
+						// Update the uecommandline.txt
+						FileReference IntermediateCmdLineFile = FileReference.Combine(SC.StageDirectory, "UECommandLine.txt");
+						Log.TraceInformation("Writing cmd line to: " + IntermediateCmdLineFile.FullName);
+						Project.WriteStageCommandline(IntermediateCmdLineFile, Params, SC);
+					}
+				}
+
+				if (SeparateAssetPackaging)
+				{
+					PackagePakFiles(Params, SC, OutputNameBase);
+				}
+				else
+				{
+					UpdateCodePackagesWithData(Params, SC, OutputNameBase);
+				}
+
+				MakeBundle(Params, SC, OutputNameBase, SeparateAssetPackaging);
+				
+				if (bShouldCreateAppInstaller)
+				{
+					MakeAppInstaller(Params, SC, OutputNameBase, autoUpdateURL, hoursBetweenUpdates);
+				}
+			}
+
+
+			CopyVCLibs(Params, SC);
+
+			// If the user indicated that they will distribute this build, then let's also generate an
+			// appxupload file suitable for submission to the Windows Store.  This file zips together
+			// the appx package and the public symbols (themselves zipped) for the binaries.
+			if (Params.Distribution)
+			{
+				List<FileReference> SymbolFilesToZip = new List<FileReference>();
+				DirectoryReference StageDirRef = new DirectoryReference(SC.StageDirectory.FullName);
+				DirectoryReference PublicSymbols = DirectoryReference.Combine(StageDirRef, "PublicSymbols");
+				DirectoryReference.CreateDirectory(PublicSymbols);
+				foreach (StageTarget Target in SC.StageTargets)
+				{
+					foreach (BuildProduct Product in Target.Receipt.BuildProducts)
+					{
+						if (Product.Type == BuildProductType.SymbolFile)
+						{
+							FileReference FullSymbolFile = new FileReference(Product.Path.FullName);
+							FileReference TempStrippedSymbols = FileReference.Combine(PublicSymbols, FullSymbolFile.GetFileName());
+							StripSymbols(FullSymbolFile, TempStrippedSymbols);
+							SymbolFilesToZip.Add(TempStrippedSymbols);
+						}
+					}
+				}
+				FileReference AppxSymFile = FileReference.Combine(StageDirRef, OutputNameBase + Extension + "sym");
+				ZipFiles(AppxSymFile, PublicSymbols, SymbolFilesToZip);
+
+				FileReference AppxUploadFile = FileReference.Combine(StageDirRef, OutputNameBase + Extension + "upload");
+				ZipFiles(AppxUploadFile, StageDirRef,
+					new FileReference[]
+					{
+							new FileReference(OutputAppX),
+							AppxSymFile
+					});
+			}
+		}
+
+		public override IProcessResult RunClient(ERunOptions ClientRunFlags, string ClientApp, string ClientCmdLine, ProjectParams Params)
+		{
+			IProcessResult ProcResult = null;
+			foreach (string DeviceAddress in Params.DeviceNames)
+			{
+				// Prefer launcher tool for local device since it avoids having to deal with certificate issues on the
+				// device portal connection.
+				if (IsLocalDevice(DeviceAddress))
+				{
+					ProcResult = RunUsingLauncherTool(DeviceAddress, ClientRunFlags, ClientApp, ClientCmdLine, Params);
+				}
+				else
+				{
+					ProcResult = RunUsingDevicePortal(DeviceAddress, ClientRunFlags, ClientApp, ClientCmdLine, Params);
+				}
+			}
+			return ProcResult;
+		}
+
+        public override List<FileReference> GetExecutableNames(DeploymentContext SC)
+		{
+			// If we're calling this for the purpose of running the app then the string we really
+			// need is the AUMID.  We can't form a full AUMID here without making assumptions about 
+			// how the PFN is built, which (while straightforward) does not appear to be officially
+			// documented.  So we'll save off the path to the manifest, which the launch process can
+			// parse later for information that, in conjunction with the target device, will allow
+			// for looking up the true AUMID.
+			List<FileReference> Exes = new List<FileReference>();
+			Exes.Add(new FileReference(Path.Combine(SC.StageDirectory.FullName, "AppxManifest.xml")));
+			return Exes;
+		}
+
+		public override bool IsSupported { get { return true; } }
+		public override bool UseAbsLog { get { return false; } }
+		public override bool LaunchViaUFE { get { return false; } }
+
+		public override List<string> GetDebugFileExtensions()
+		{
+			return new List<string> { ".pdb", ".map" };
+		}
+
+		public override void StripSymbols(FileReference SourceFile, FileReference TargetFile)
+		{
+			bool bStripInPlace = false;
+
+			if (SourceFile == TargetFile)
+			{
+				// PDBCopy only supports creation of a brand new stripped file so we have to create a temporary filename
+				TargetFile = new FileReference(Path.Combine(TargetFile.Directory.FullName, Guid.NewGuid().ToString() + TargetFile.GetExtension()));
+				bStripInPlace = true;
+			}
+
+			ProcessStartInfo StartInfo = new ProcessStartInfo();
+
+			if (PDBCopyPath == null || !FileReference.Exists(PDBCopyPath))
+			{
+				throw new AutomationException(ExitCode.Error_SDKNotFound, "Debugging Tools for Windows aren't installed. Please follow https://docs.microsoft.com/en-us/windows-hardware/drivers/debugger/debugger-download-tools");
+			}
+
+			StartInfo.FileName = PDBCopyPath.FullName;
+			StartInfo.Arguments = String.Format("\"{0}\" \"{1}\" -p", SourceFile.FullName, TargetFile.FullName);
+			StartInfo.UseShellExecute = false;
+			StartInfo.CreateNoWindow = true;
+			Utils.RunLocalProcessAndLogOutput(StartInfo);
+
+			if (bStripInPlace)
+			{
+				// Copy stripped file to original location and delete the temporary file
+				File.Copy(TargetFile.FullName, SourceFile.FullName, true);
+				FileReference.Delete(TargetFile);
+			}
+		}
+
+		private void GenerateSigningCertificate(string InCertificatePath, string InPublisher)
+		{
+			// Ensure the output directory exists otherwise makecert will fail.
+			InternalUtils.SafeCreateDirectory(Path.GetDirectoryName(InCertificatePath));
+
+			// MakeCert.exe -r -h 0 -n "CN=No Publisher, O=No Publisher" -eku 1.3.6.1.5.5.7.3.3 -pe -sv "Signing Certificate.pvk" "Signing Certificate.cer"
+			// pvk2pfx -pvk "Signing Certificate.pvk" -spc "Signing Certificate.cer" -pfx "Signing Certificate.pfx"
+			string CerFile = Path.ChangeExtension(InCertificatePath, ".cer");
+			string PvkFile = Path.ChangeExtension(InCertificatePath, ".pvk");
+
+			string MakeCertCommandLine = string.Format(@"-r -h 0 -n ""{0}"" -eku 1.3.6.1.5.5.7.3.3 -pe -sv ""{1}"" ""{2}""", InPublisher, PvkFile, CerFile);
+			RunAndLog(CmdEnv, MakeCertPath.FullName, MakeCertCommandLine, null, 0, null, ERunOptions.None);
+
+			string Pvk2PfxCommandLine = string.Format(@"-pvk ""{0}"" -spc ""{1}"" -pfx ""{2}""", PvkFile, CerFile, InCertificatePath);
+			RunAndLog(CmdEnv, Pvk2PfxPath.FullName, Pvk2PfxCommandLine, null, 0, null, ERunOptions.None);
+		}
+
+		private bool IsLocalDevice(string DeviceAddress)
+		{
+			Uri uriResult;
+			bool result = Uri.TryCreate(DeviceAddress, UriKind.Absolute, out uriResult)
+				&& (uriResult.Scheme == Uri.UriSchemeHttp || uriResult.Scheme == Uri.UriSchemeHttps);
+			return !result;
+		}
+
+		private void WaitFor(IAsyncOperationWithProgress<DeploymentResult, DeploymentProgress> deploymentOperation)
+		{
+			// This event is signaled when the operation completes
+			ManualResetEvent opCompletedEvent = new ManualResetEvent(false);
+
+			// Define the delegate using a statement lambda
+			deploymentOperation.Completed = (depProgress, status) => { opCompletedEvent.Set(); };
+
+			// Wait until the operation completes
+			opCompletedEvent.WaitOne();
+
+			// Check the status of the operation
+			if (deploymentOperation.Status == AsyncStatus.Error)
+			{
+				DeploymentResult deploymentResult = deploymentOperation.GetResults();
+				LogInformation("Error code: {0}", deploymentOperation.ErrorCode);
+				LogInformation("Error text: {0}", deploymentResult.ErrorText);
+				throw new AutomationException(ExitCode.Error_AppInstallFailed, deploymentResult.ErrorText, "");
+			}
+			else if (deploymentOperation.Status == AsyncStatus.Canceled)
+			{
+				//LogInformation("Operation canceled");
+				throw new AutomationException(ExitCode.Error_AppInstallFailed, "Operation canceled", "");
+			}
+			else if (deploymentOperation.Status == AsyncStatus.Completed)
+			{
+				//LogInformation("Operation succeeded");
+			}
+			else
+			{
+				//LogInformation("Operation status unknown");
+			}
+		}
+
+		private void DeployToLocalDevice(ProjectParams Params, DeploymentContext SC)
+		{
+            if (!RuntimePlatform.IsWindows)
+            {
+                return;
+            }
+
+            bool bRequiresPackage = Params.Package || SC.StageTargetPlatform.RequiresPackageToDeploy;
+			string Name;
+			string Publisher;
+			GetPackageInfo(Params, out Name, out Publisher);
+			PackageManager PackMgr = new PackageManager();
+			try
+			{
+				var ExistingPackage = PackMgr.FindPackagesForUser("", Name, Publisher).FirstOrDefault();
+
+				// Only remove an existing package if it's in development mode; otherwise the removal might silently delete stuff
+				// that the user wanted.
+				if (ExistingPackage != null)
+				{
+					if (ExistingPackage.IsDevelopmentMode)
+					{
+						WaitFor(PackMgr.RemovePackageAsync(ExistingPackage.Id.FullName, RemovalOptions.PreserveApplicationData));
+					}
+					else if (!Params.Package)
+					{
+						throw new AutomationException(ExitCode.Error_AppInstallFailed, "A packaged version of the application already exists.  It must be uninstalled manually - note this will remove user data.");
+					}
+				}
+
+				{
+					string PackagePath = Path.Combine(SC.StageDirectory.FullName, "AppxManifest.xml");
+
+
+					LogInformation(String.Format("Registring {0}", PackagePath));
+					WaitFor(PackMgr.RegisterPackageAsync(new Uri(PackagePath), null, DeploymentOptions.DevelopmentMode));
+				}
+			}
+			catch (AggregateException agg)
+			{
+				throw new AutomationException(ExitCode.Error_AppInstallFailed, agg.InnerException, "");
+			}
+
+			// Package should now be installed.  Locate it and make sure it's permitted to connect over loopback.
+			try
+			{
+				var InstalledPackage = PackMgr.FindPackagesForUser("", Name, Publisher).FirstOrDefault();
+				string LoopbackExemptCmdLine = string.Format("loopbackexempt -a -n={0}", InstalledPackage.Id.FamilyName);
+				RunAndLog(CmdEnv, "checknetisolation.exe", LoopbackExemptCmdLine, null, 0, null, ERunOptions.None);
+			}
+			catch
+			{
+				LogWarning("Failed to apply a loopback exemption to the deployed app.  Connection to a local cook server will fail.");
+			}
+		}
+
+		private WindowsArchitecture RemoteDeviceArchitecture(string DeviceAddress, ProjectParams Params)
+		{
+			string OsVersionString = "";
+			try
+			{
+				DefaultDevicePortalConnection conn = new DefaultDevicePortalConnection(DeviceAddress, Params.DeviceUsername, Params.DevicePassword);
+				DevicePortal portal = new DevicePortal(conn);
+				portal.ConnectionStatus += Portal_ConnectionStatus;
+				//portal.ConnectionStatus += (sender, args) => { Console.WriteLine(String.Format("Connected to portal with {0}", args.Status)); }; 
+
+				portal.ConnectAsync().Wait(); //think about adding a cert here!
+				var osInfo = portal.GetOperatingSystemInformationAsync().Result;
+				OsVersionString = osInfo.OsVersionString;
+			}
+			catch (AggregateException e)
+			{
+				if (e.InnerException is AutomationException)
+				{
+					throw e.InnerException;
+				}
+				else
+				{
+					throw new AutomationException(ExitCode.Error_AppInstallFailed, e.InnerException, e.InnerException.Message);
+				}
+			}
+			catch (Exception e)
+			{
+				throw new AutomationException(ExitCode.Error_AppInstallFailed, e, e.Message);
+			}
+
+
+			string[] osBlocks = OsVersionString.Split('.');
+			if (osBlocks.Length < 3)
+			{
+				LogError(String.Format("Wrong OS version string {0} from {1} device", OsVersionString, DeviceAddress));
+				throw new AutomationException(ExitCode.Error_AppInstallFailed, "Wrong OS version string");
+			}
+
+			string osArchBlock = osBlocks[2].ToLower();
+
+			if (osArchBlock.StartsWith("amd64"))
+			{
+				return WindowsArchitecture.x64;
+			}
+			else if (osArchBlock.StartsWith("arm64"))
+			{
+				return WindowsArchitecture.ARM64;
+			}
+			else
+			{
+				LogError(String.Format("Unsupported OS architecture {0} from {1} device", osArchBlock, DeviceAddress));
+				throw new AutomationException(ExitCode.Error_AppInstallFailed, "Unsupported OS architecture");
+			}
+		}
+
+		private void DeployToRemoteDevice(string DeviceAddress, ProjectParams Params, DeploymentContext SC)
+		{
+			if (!RuntimePlatform.IsWindows)
+            {
+                return;
+            }
+
+			string Name;
+			string Publisher;
+			GetPackageInfo(Params, out Name, out Publisher);
+
+			if (Params.Package || SC.StageTargetPlatform.RequiresPackageToDeploy)
+			{
+				DefaultDevicePortalConnection conn = new DefaultDevicePortalConnection(DeviceAddress, Params.DeviceUsername, Params.DevicePassword);
+				DevicePortal portal = new DevicePortal(conn);
+				portal.UnvalidatedCert += (sender, certificate, chain, sslPolicyErrors) =>
+				{
+					return ShouldAcceptCertificate(new X509Certificate2(certificate), Params.Unattended);
+				};
+				portal.ConnectionStatus += Portal_ConnectionStatus;
+				try
+				{
+					portal.ConnectAsync().Wait();
+					string PackageName = Params.ShortProjectName;
+					string PackagePath = Path.Combine(SC.StageDirectory.FullName, PackageName + Extension + "bundle");
+					string CertPath = Path.Combine(SC.StageDirectory.FullName, PackageName + ".cer");
+
+					{
+						var list = portal.GetInstalledAppPackagesAsync().Result;
+						//LogInformation(String.Format("Current Name = {0}, Publisher = {1}", Name, Publisher)); //code for installation debugging 
+						//foreach(var p in list.Packages)
+						//{
+						//	LogInformation(String.Format("Package FamilyName = {0}, FullName = {1}, Publisher = {2}", p.FamilyName, p.FullName, p.Publisher));
+						//}
+						foreach (var pack in list.Packages.FindAll(package => package.FamilyName.StartsWith(Name) && package.Publisher == Publisher))
+						{
+							portal.UninstallApplicationAsync(pack.FullName).Wait();
+						}
+					}
+
+					List<string> Dependencies = new List<string>();
+					bool UseDebugCrt = false;
+					WindowsCompiler Compiler = WindowsCompiler.Default;
+					TargetRules Rules = null;
+					if (Params.HasGameTargetDetected)
+					{
+						//Rules = Params.ProjectTargets[TargetType.Game].Rules;
+						Rules = Params.ProjectTargets.Find(x => x.Rules.Type == TargetType.Game).Rules;
+					}
+					else if (Params.HasClientTargetDetected)
+					{
+						Rules = Params.ProjectTargets.Find(x => x.Rules.Type == TargetType.Game).Rules;
+					}
+
+					if (Rules != null)
+					{
+						UseDebugCrt = Params.ClientConfigsToBuild.Contains(UnrealTargetConfiguration.Debug) && Rules.bDebugBuildsActuallyUseDebugCRT;
+						Compiler = Rules.WindowsPlatform.Compiler;
+					}
+
+					Dependencies.AddRange(GetPathToVCLibsPackages(UseDebugCrt, Compiler));
+
+					portal.AppInstallStatus += Portal_AppInstallStatus;
+					portal.InstallApplicationAsync(string.Empty, PackagePath, Dependencies, CertPath).Wait();
+				}
+				catch (AggregateException e)
+				{
+					if (e.InnerException is AutomationException)
+					{
+						throw e.InnerException;
+					}
+					else
+					{
+						throw new AutomationException(ExitCode.Error_AppInstallFailed, e.InnerException, e.InnerException.Message);
+					}
+				}
+				catch (Exception e)
+				{
+					throw new AutomationException(ExitCode.Error_AppInstallFailed, e, e.Message);
+				}
+			}
+			else
+			{
+				throw new AutomationException(ExitCode.Error_AppInstallFailed, "Remote deployment of unpackaged apps is not supported.");
+			}
+		}
+
+
+		private IProcessResult RunUsingLauncherTool(string DeviceAddress, ERunOptions ClientRunFlags, string ClientApp, string ClientCmdLine, ProjectParams Params)
+		{
+            if (!RuntimePlatform.IsWindows)
+            {
+                return null;
+            }
+			return RunUsingLauncherToolAsync(Params, ClientRunFlags).Result;
+		}
+
+
+		// Currently does not build on UE4 build servers. Temporarily blocking out until resolved.
+#if false
+		private async Task<HoloLensLauncherCreatedProcess> RunUsingLauncherToolAsync(ProjectParams Params, ERunOptions ClientRunFlags)
+		{
+			string Name;
+			string Publisher;
+			GetPackageInfo(Params, out Name, out Publisher);
+
+			PackageManager PackMgr = new PackageManager();
+			Windows.ApplicationModel.Package InstalledPackage = PackMgr.FindPackagesForUser("", Name, Publisher).FirstOrDefault();
+
+			if (InstalledPackage == null)
+			{
+				throw new AutomationException(ExitCode.Error_LauncherFailed, "Could not find installed app (Name: {0}, Publisher: {1}", Name, Publisher);
+			}
+
+			IReadOnlyList<AppListEntry> entries = await InstalledPackage.GetAppListEntriesAsync();
+
+			if (entries.Count == 0)
+			{
+				throw new AutomationException(ExitCode.Error_LauncherFailed, "Could not find appentry in installed app (Name: {0}, Publisher: {1}", Name, Publisher);
+			}
+
+			AppListEntry entry = entries.First();
+			if (!await entry.LaunchAsync())
+			{
+				return null;
+			}
+
+			string LogFile;
+			if (Params.CookOnTheFly)
+			{
+				LogFile = Path.Combine(Params.RawProjectPath.Directory.FullName, "Saved", "Cooked", PlatformType.ToString(), Params.ShortProjectName, "HoloLensLocalAppData", Params.ShortProjectName, "Saved", "Logs", Params.ShortProjectName + ".log");
+			}
+			else
+			{
+				LogFile = Path.Combine(Environment.GetFolderPath(Environment.SpecialFolder.LocalApplicationData), "Packages", InstalledPackage.Id.FamilyName, "LocalState", Params.ShortProjectName, "Saved", "Logs", Params.ShortProjectName + ".log");
+			}
+
+			Process[] Procs = Process.GetProcessesByName(Params.ShortProjectName + Params.SpecifiedArchitecture + ".exe");
+			Process Proc = null;
+			foreach (var P in Procs)
+			{
+				DirectoryReference WorkDir = new DirectoryReference(P.StartInfo.WorkingDirectory);
+				if (WorkDir.IsUnderDirectory(Params.RawProjectPath.Directory))
+				{
+					Proc = P;
+					break;
+				}
+			}
+
+			if (Proc == null)
+			{
+				return null;
+			}
+
+			bool AllowSpew = ClientRunFlags.HasFlag(ERunOptions.AllowSpew);
+			LogEventType SpewVerbosity = ClientRunFlags.HasFlag(ERunOptions.SpewIsVerbose) ? LogEventType.Verbose : LogEventType.Console;
+			HoloLensLauncherCreatedProcess UwpProcessResult = new HoloLensLauncherCreatedProcess(Proc, LogFile, AllowSpew, SpewVerbosity);
+			ProcessManager.AddProcess(UwpProcessResult);
+
+			ProcessManager.AddProcess(UwpProcessResult);
+			if (!ClientRunFlags.HasFlag(ERunOptions.NoWaitForExit))
+			{
+				UwpProcessResult.WaitForExit();
+				UwpProcessResult.OnProcessExited();
+				UwpProcessResult.DisposeProcess();
+			}
+			return UwpProcessResult;
+		}
+#else
+		private Task<HoloLensLauncherCreatedProcess> RunUsingLauncherToolAsync(ProjectParams Params, ERunOptions ClientRunFlags)
+		{
+			return null;
+		}
+#endif
+
+		private IProcessResult RunUsingDevicePortal(string DeviceAddress, ERunOptions ClientRunFlags, string ClientApp, string ClientCmdLine, ProjectParams Params)
+		{
+            if (!RuntimePlatform.IsWindows)
+            {
+                return null;
+            }
+
+            string Name;
+			string Publisher;
+			GetPackageInfo(Params, out Name, out Publisher);
+
+			DefaultDevicePortalConnection conn = new DefaultDevicePortalConnection(DeviceAddress, Params.DeviceUsername, Params.DevicePassword);
+			DevicePortal portal = new DevicePortal(conn);
+			portal.UnvalidatedCert += (sender, certificate, chain, sslPolicyErrors) =>
+			{
+				return ShouldAcceptCertificate(new System.Security.Cryptography.X509Certificates.X509Certificate2(certificate), Params.Unattended);
+			};
+
+			try
+			{
+				portal.ConnectAsync().Wait();
+				string Aumid = string.Empty;
+				string FullName = string.Empty;
+				var AllAppsTask = portal.GetInstalledAppPackagesAsync();
+				AllAppsTask.Wait();
+				foreach (var App in AllAppsTask.Result.Packages)
+				{
+					// App.Name seems to report the localized name.
+					if (App.FamilyName.StartsWith(Name) && App.Publisher == Publisher)
+					{
+						Aumid = App.AppId;
+						FullName = App.FullName;
+						break;
+					}
+				}
+
+				var Result = new HoloLensDevicePortalCreatedProcess(portal, FullName, Params.ShortProjectName);
+				var LaunchTask = portal.LaunchApplicationAsync(Aumid, FullName);
+
+				// Message back to the UE4 Editor to correctly set the app id for each device
+				Console.WriteLine("Running Package@Device:{0}@{1}", FullName, DeviceAddress);
+
+				LaunchTask.Wait();
+				Result.SetProcId(LaunchTask.Result);
+
+				ProcessManager.AddProcess(Result);
+				if (!ClientRunFlags.HasFlag(ERunOptions.NoWaitForExit))
+				{
+					Result.WaitForExit();
+					Result.OnProcessExited();
+					Result.DisposeProcess();
+				}
+				return Result;
+			}
+			catch (AggregateException e)
+			{
+				if (e.InnerException is AutomationException)
+				{
+					throw e.InnerException;
+				}
+				else
+				{
+					throw new AutomationException(ExitCode.Error_LauncherFailed, e.InnerException, e.InnerException.Message);
+				}
+			}
+			catch (Exception e)
+			{
+				throw new AutomationException(ExitCode.Error_LauncherFailed, e, e.Message);
+			}
+		}
+
+		private void GetPackageInfo(ProjectParams Params, out string Name, out string Publisher)
+		{
+			ConfigHierarchy PlatformEngineConfig = null;
+			Params.GameConfigs.TryGetValue(PlatformType, out PlatformEngineConfig);
+
+			if (PlatformEngineConfig == null || !PlatformEngineConfig.GetString("/Script/EngineSettings.GeneralProjectSettings", "ProjectName", out Name))
+			{
+				Name = "DefaultUE4Project";
+			}
+
+			Name = Regex.Replace(Name, "[^-.A-Za-z0-9]", "");
+
+			if (PlatformEngineConfig == null || !PlatformEngineConfig.GetString("/Script/EngineSettings.GeneralProjectSettings", "CompanyDistinguishedName", out Publisher))
+			{
+				Publisher = "CN=NoPublisher";
+			}
+		}
+
+		public override void GetFilesToArchive(ProjectParams Params, DeploymentContext SC)
+		{
+			string OutputNameBase = Params.HasDLCName ? Params.DLCFile.GetFileNameWithoutExtension() : Params.ShortProjectName;
+			string PackagePath = SC.StageDirectory.FullName;
+			SC.ArchiveFiles(PackagePath, OutputNameBase + Extension + "bundle");
+			SC.ArchiveFiles(PackagePath, OutputNameBase + ".cer");
+			SC.ArchiveFiles(PackagePath, OutputNameBase + Extension + "upload");
+
+			SC.ArchiveFiles(PackagePath, "*VCLibs*.appx");
+			SC.ArchiveFiles(PackagePath, OutputNameBase + ".appinstaller");
+		}
+
+		private bool ShouldAcceptCertificate(System.Security.Cryptography.X509Certificates.X509Certificate2 Certificate, bool Unattended)
+		{
+            if (!RuntimePlatform.IsWindows)
+            {
+                return false;
+            }
+
+            if (AcceptThumbprints.Contains(Certificate.Thumbprint))
+			{
+				return true;
+			}
+
+			if (Unattended)
+			{
+				throw new AutomationException(ExitCode.Error_CertificateNotFound, "Cannot connect to remote device: certificate is untrusted and cannot prompt for consent (running unattended).");
+			}
+
+			System.Windows.Forms.DialogResult AcceptResult = System.Windows.Forms.MessageBox.Show(
+				string.Format("Do you want to accept the following certificate?\n\nThumbprint:\n\t{0}\nIssues:\n\t{1}", Certificate.Thumbprint, Certificate.Issuer),
+				"Untrusted Certificate Detected",
+				System.Windows.Forms.MessageBoxButtons.YesNo,
+				System.Windows.Forms.MessageBoxIcon.Question,
+				System.Windows.Forms.MessageBoxDefaultButton.Button2);
+			if (AcceptResult == System.Windows.Forms.DialogResult.Yes)
+			{
+				AcceptThumbprints.Add(Certificate.Thumbprint);
+				return true;
+			}
+
+			throw new AutomationException(ExitCode.Error_CertificateNotFound, "Cannot connect to remote device: certificate is untrusted and user declined to accept.");
+		}
+
+		private void Portal_AppInstallStatus(Microsoft.Tools.WindowsDevicePortal.DevicePortal sender, Microsoft.Tools.WindowsDevicePortal.ApplicationInstallStatusEventArgs args)
+		{
+			if (args.Status == Microsoft.Tools.WindowsDevicePortal.ApplicationInstallStatus.Failed)
+			{
+				throw new AutomationException(ExitCode.Error_AppInstallFailed, args.Message);
+			}
+			else
+			{
+				LogLog(args.Message);
+			}
+		}
+
+		private void Portal_ConnectionStatus(Microsoft.Tools.WindowsDevicePortal.DevicePortal sender, Microsoft.Tools.WindowsDevicePortal.DeviceConnectionStatusEventArgs args)
+		{
+			if (args.Status == Microsoft.Tools.WindowsDevicePortal.DeviceConnectionStatus.Failed)
+			{
+				throw new AutomationException(args.Message);
+			}
+			else
+			{
+				LogLog(args.Message);
+			}
+		}
+
+		private string[] GetPathToVCLibsPackages(bool UseDebugCrt, WindowsCompiler Compiler)
+		{
+			List<DirectoryReference> SdkRootDirs = new List<DirectoryReference>();
+			WindowsExports.EnumerateSdkRootDirs(SdkRootDirs);
+			
+			string VCVersionFragment;
+            switch (Compiler)
+			{
+				case WindowsCompiler.VisualStudio2019:
+				case WindowsCompiler.VisualStudio2022:
+				//Compiler version is still 14 for 2017
+				case WindowsCompiler.Default:
+					VCVersionFragment = "14";
+					break;
+
+				default:
+					VCVersionFragment = "Unsupported_VC_Version";
+					break;
+			}
+
+			List<string> Runtimes = new List<string>();
+
+			bool foundVcLibs = false;
+			foreach (DirectoryReference SdkRootDir in SdkRootDirs)
+			{
+				if (!foundVcLibs)
+				{
+					foreach (var Arch in ActualArchitectures)
+					{
+						string ArchitectureFragment = WindowsExports.GetArchitectureSubpath(Arch);
+
+						// For whatever reason, VCLibs are not in Windows Kits/, but Microsoft SDKs/Windows Kits, so go up two directories and start from there.
+						if (DirectoryExists(SdkRootDir.ParentDirectory.ToString()) && DirectoryExists(SdkRootDir.ParentDirectory.ParentDirectory.ToString()))
+						{
+							DirectoryReference Win64SDKsRootDir = SdkRootDir.ParentDirectory.ParentDirectory;
+
+							string RuntimePath = Path.Combine(Win64SDKsRootDir.ToString(),
+								"Microsoft SDKs",
+								"Windows Kits",
+								"10",
+								"ExtensionSDKs",
+								"Microsoft.VCLibs",
+								string.Format("{0}.0", VCVersionFragment),
+								"Appx",
+								UseDebugCrt ? "Debug" : "Retail",
+								ArchitectureFragment,
+								string.Format("Microsoft.VCLibs.{0}.{1}.00.appx", ArchitectureFragment, VCVersionFragment));
+
+							if (FileExists(RuntimePath))
+							{
+								Runtimes.Add(RuntimePath);
+								LogLog("found vclibs, path: " + RuntimePath);
+								foundVcLibs = true;
+							}
+							else
+							{
+								LogLog("no vclib found, path: " + RuntimePath);
+							}
+						}
+						else
+						{
+							LogLog("Directory structure for VCLibs not present , SDK path: " + SdkRootDir.ToString());
+						}
+					}
+				}
+			}
+
+			if (!foundVcLibs)
+			{
+				LogError("no appropriate VCLibs found in any of the SDK locations");
+			}
+
+			return Runtimes.ToArray();
+		}
+
+		private void GenerateDLCManifestIfNecessary(ProjectParams Params, DeploymentContext SC)
+		{
+			// Only required for DLC
+			if (!Params.HasDLCName)
+			{
+				return;
+			}
+
+			// Only required for the first stage (package or deploy) that requires a manifest.
+			// Assumes that the staging directory is pre-cleaned
+			if (FileReference.Exists(FileReference.Combine(SC.StageDirectory, "AppxManifest.xml")))
+			{
+				return;
+			}
+
+
+			HoloLensExports.CreateManifestForDLC(Params.DLCFile, SC.StageDirectory);
+		}
+
+		private static List<string> AcceptThumbprints = new List<string>();
+	}
+}
diff --git a/Engine/Platforms/Hololens/Source/Programs/AutomationTool/Properties/AssemblyInfo.cs b/Engine/Platforms/Hololens/Source/Programs/AutomationTool/Properties/AssemblyInfo.cs
new file mode 100644
--- /dev/null
+++ b/Engine/Platforms/Hololens/Source/Programs/AutomationTool/Properties/AssemblyInfo.cs
@@ -0,0 +1,12 @@
+﻿// Copyright Epic Games, Inc. All Rights Reserved.
+
+using System.Reflection;
+using System.Runtime.CompilerServices;
+using System.Runtime.InteropServices;
+
+// General Information about an assembly is controlled through the following 
+// set of attributes. Change these attribute values to modify the information
+// associated with an assembly.
+[assembly: AssemblyTitle("HoloLens.Automation")]
+[assembly: AssemblyDescription("")]
+[assembly: AssemblyConfiguration("")]
diff --git a/Engine/Platforms/Hololens/Source/Programs/UnrealBuildTool/HoloLensDeploy.cs b/Engine/Platforms/Hololens/Source/Programs/UnrealBuildTool/HoloLensDeploy.cs
new file mode 100644
--- /dev/null
+++ b/Engine/Platforms/Hololens/Source/Programs/UnrealBuildTool/HoloLensDeploy.cs
@@ -0,0 +1,555 @@
+// Copyright Epic Games, Inc. All Rights Reserved.
+
+using System;
+using System.Collections.Generic;
+using System.Text;
+using System.IO;
+using System.Diagnostics;
+using System.Linq;
+using EpicGames.Core;
+using UnrealBuildBase;
+using System.Runtime.Versioning;
+
+namespace UnrealBuildTool
+{
+	/// <summary>
+	///  Base class to handle deploy of a target for a given platform
+	/// </summary>
+	[SupportedOSPlatform("windows")]
+	class HoloLensDeploy : UEBuildDeploy
+	{
+		private FileReference? MakeAppXPath;
+		private FileReference? SignToolPath;
+		private string Extension = ".appx";
+
+		/// <summary>
+		/// Utility function to delete a file
+		/// </summary>
+		void DeployHelper_DeleteFile(string InFileToDelete)
+		{
+			Log.TraceInformation("HoloLensDeploy.DeployHelper_DeleteFile({0})", InFileToDelete);
+			if (File.Exists(InFileToDelete) == true)
+			{
+				FileAttributes attributes = File.GetAttributes(InFileToDelete);
+				if ((attributes & FileAttributes.ReadOnly) == FileAttributes.ReadOnly)
+				{
+					attributes &= ~FileAttributes.ReadOnly;
+					File.SetAttributes(InFileToDelete, attributes);
+				}
+				File.Delete(InFileToDelete);
+			}
+		}
+
+		/// <summary>
+		/// Copy the contents of the given source directory to the given destination directory
+		/// </summary>
+		bool CopySourceToDestDir(string InSourceDirectory, string InDestinationDirectory, string InWildCard,
+			bool bInIncludeSubDirectories, bool bInRemoveDestinationOrphans)
+		{
+			Log.TraceInformation("HoloLensDeploy.CopySourceToDestDir({0}, {1}, {2},...)", InSourceDirectory, InDestinationDirectory, InWildCard);
+			if (Directory.Exists(InSourceDirectory) == false)
+			{
+				Log.TraceInformation("Warning: CopySourceToDestDir - SourceDirectory does not exist: {0}", InSourceDirectory);
+				return false;
+			}
+
+			// Make sure the destination directory exists!
+			Directory.CreateDirectory(InDestinationDirectory);
+
+			SearchOption OptionToSearch = bInIncludeSubDirectories ? SearchOption.AllDirectories : SearchOption.TopDirectoryOnly;
+
+			var SourceDirs = new List<string>(Directory.GetDirectories(InSourceDirectory, "*.*", OptionToSearch));
+			foreach (string SourceDir in SourceDirs)
+			{
+				string SubDir = SourceDir.Replace(InSourceDirectory, "");
+				string DestDir = InDestinationDirectory + SubDir;
+				Directory.CreateDirectory(DestDir);
+			}
+
+			var SourceFiles = new List<string>(Directory.GetFiles(InSourceDirectory, InWildCard, OptionToSearch));
+			var DestFiles = new List<string>(Directory.GetFiles(InDestinationDirectory, InWildCard, OptionToSearch));
+
+			// Keep a list of the files in the source directory... without the source path
+			List<string> FilesInSource = new List<string>();
+
+			// Copy all the source files that are newer...
+			foreach (string SourceFile in SourceFiles)
+			{
+				string Filename = SourceFile.Replace(InSourceDirectory, "");
+				FilesInSource.Add(Filename.ToUpperInvariant());
+				string DestFile = InDestinationDirectory + Filename;
+
+				System.DateTime SourceTime = File.GetLastWriteTime(SourceFile);
+				System.DateTime DestTime = File.GetLastWriteTime(DestFile);
+
+				if (SourceTime > DestTime)
+				{
+					try
+					{
+						DeployHelper_DeleteFile(DestFile);
+						File.Copy(SourceFile, DestFile, true);
+					}
+					catch (Exception exceptionMessage)
+					{
+						Log.TraceInformation("Failed to copy {0} to deployment: {1}", SourceFile, exceptionMessage);
+					}
+				}
+			}
+
+			if (bInRemoveDestinationOrphans == true)
+			{
+				// If requested, delete any destination files that do not have a corresponding
+				// file in the source directory
+				foreach (string DestFile in DestFiles)
+				{
+					string DestFilename = DestFile.Replace(InDestinationDirectory, "");
+					if (FilesInSource.Contains(DestFilename.ToUpperInvariant()) == false)
+					{
+						Log.TraceInformation("Destination file does not exist in Source - DELETING: {0}", DestFile);
+						//FileAttributes attributes = File.GetAttributes(DestFile);
+						try
+						{
+							DeployHelper_DeleteFile(DestFile);
+						}
+						catch (Exception exceptionMessage)
+						{
+							Log.TraceInformation("Failed to delete {0} from deployment: {1}", DestFile, exceptionMessage);
+						}
+					}
+				}
+			}
+
+			return true;
+		}
+
+
+		/// <summary>
+		/// Helper function for copying files
+		/// </summary>
+		void CopyFile(string InSource, string InDest, bool bForce)
+		{
+			if (File.Exists(InSource) == true)
+			{
+				if (File.Exists(InDest) == true)
+				{
+					if (File.GetLastWriteTime(InSource).CompareTo(File.GetLastWriteTime(InDest)) == 0)
+					{
+						//If the source and dest have the file and they have the same write times they are assumed to be equal and we don't need to copy.
+						return;
+					}
+					if (bForce == true)
+					{
+						DeployHelper_DeleteFile(InDest);
+					}
+				}
+				Log.TraceInformation("HoloLensDeploy.CopyFile({0}, {1}, {2})", InSource, InDest, bForce);
+				File.Copy(InSource, InDest, true);
+				File.SetAttributes(InDest, File.GetAttributes(InDest) & ~FileAttributes.ReadOnly);
+			}
+			else
+			{
+				Log.TraceInformation("HoloLensDeploy: File didn't exist - {0}", InSource);
+			}
+		}
+
+		/// <summary>
+		/// Helper function for copying a tree files
+		/// </summary>
+		void CopyDirectory(string InSource, string InDest, bool bForce, bool bRecurse)
+		{
+			if (Directory.Exists(InSource))
+			{
+				if (!Directory.Exists(InDest))
+				{
+					Directory.CreateDirectory(InDest);
+				}
+
+				// Copy all files
+				string[] FilesInDir = Directory.GetFiles(InSource);
+				foreach (string FileSourcePath in FilesInDir)
+				{
+					string FileDestPath = Path.Combine(InDest, Path.GetFileName(FileSourcePath));
+					CopyFile(FileSourcePath, FileDestPath, true);
+				}
+
+				// Recurse sub directories
+				string[] DirsInDir = Directory.GetDirectories(InSource);
+				foreach (string DirSourcePath in DirsInDir)
+				{
+					string DirName = Path.GetFileName(DirSourcePath);
+					string DirDestPath = Path.Combine(InDest, DirName);
+					CopyDirectory(DirSourcePath, DirDestPath, bForce, bRecurse);
+				}
+			}
+		}
+
+		public bool PrepForUATPackageOrDeploy(FileReference? ProjectFile, string ProjectName, string ProjectDirectory, WindowsArchitecture Architecture, List<UnrealTargetConfiguration> TargetConfigurations, List<string> ExecutablePaths, string EngineDirectory, bool bForDistribution, string CookFlavor, bool bIsDataDeploy)
+		{
+			//@todo need to support dlc and other targets
+			//string LocalizedContentDirectory = Path.Combine(ProjectDirectory, "Content", "Localization", "Game");
+			string AbsoluteExeDirectory = Path.GetDirectoryName(ExecutablePaths[0])!;
+			//bool IsGameSpecificExe = ProjectFile != null && AbsoluteExeDirectory.StartsWith(ProjectDirectory);
+			//string RelativeExeFilePath = Path.Combine(IsGameSpecificExe ? ProjectName : "Engine", "Binaries", "HoloLens", Path.GetFileName(ExecutablePaths[0]));
+
+			//string TargetDirectory = Path.Combine(ProjectDirectory, "Saved", "HoloLens");
+
+
+			// If using a secure networking manifest, copy it to the output directory.
+			string NetworkManifest = Path.Combine(ProjectDirectory, "Config", "HoloLens", "NetworkManifest.xml");
+			if (File.Exists(NetworkManifest))
+			{
+				CopyFile(NetworkManifest, Path.Combine(AbsoluteExeDirectory, "NetworkManifest.xml"), false);
+			}
+
+			// If using Xbox Live generate the json config file expected by the SDK
+			DirectoryReference? ConfigDirRef = DirectoryReference.FromFile(ProjectFile);
+			if (ConfigDirRef == null && !string.IsNullOrEmpty(UnrealBuildTool.GetRemoteIniPath()))
+			{
+				ConfigDirRef = new DirectoryReference(UnrealBuildTool.GetRemoteIniPath()!);
+			}
+
+			MakeAppXPath = HoloLensExports.GetWindowsSdkToolPath("makeappx.exe");
+			SignToolPath = HoloLensExports.GetWindowsSdkToolPath("signtool.exe");
+
+			return true;
+		}
+
+		private void MakePackage(TargetReceipt Receipt, TargetReceipt NewReceipt, WindowsArchitecture Architecture, List<string> UpdatedFiles)
+		{
+			string OutputName = String.Format("{0}_{1}_{2}_{3}", Receipt.TargetName, Receipt.Platform, Receipt.Configuration, WindowsExports.GetArchitectureSubpath(Architecture));
+			string IntermediateDirectory = Path.Combine(Receipt.ProjectDir != null ? Receipt.ProjectDir.FullName : Unreal.EngineDirectory.FullName, "Intermediate", "Deploy", WindowsExports.GetArchitectureSubpath(Architecture));
+			string OutputDirectory = Receipt.Launch!.Directory.FullName;
+			string OutputAppX = Path.Combine(OutputDirectory, OutputName + Extension);
+			string SigningCertificate = @"Build\HoloLens\SigningCertificate.pfx";
+			string SigningCertificatePath = Path.Combine(Receipt.ProjectDir != null ? Receipt.ProjectDir.FullName : Unreal.EngineDirectory.FullName, SigningCertificate);
+			if (!File.Exists(SigningCertificatePath))
+			{
+				SigningCertificate = @"Platforms\HoloLens\Build\SigningCertificate.pfx";
+				SigningCertificatePath = Path.Combine(Receipt.ProjectDir != null ? Receipt.ProjectDir.FullName : Unreal.EngineDirectory.FullName, SigningCertificate);
+			}
+
+			string MapFilename = Path.Combine(IntermediateDirectory, OutputName + ".pkgmap");
+			var LocalRoot = Receipt.ProjectDir;
+			var EngineRoot = Unreal.RootDirectory;
+			var AddedFiles = new Dictionary<string, string>();
+			bool PackageFileNeedToBeUpdated = !File.Exists(OutputAppX);
+
+			DateTime AppXTime = DateTime.Now;
+
+			PackageFileNeedToBeUpdated = true;
+
+			if(!PackageFileNeedToBeUpdated)
+			{
+				AppXTime = File.GetLastWriteTimeUtc(OutputAppX);
+			}
+
+			{
+				foreach (var Product in Receipt.BuildProducts)
+				{
+					if (Product.Type == BuildProductType.Executable || Product.Type == BuildProductType.DynamicLibrary || Product.Type == BuildProductType.RequiredResource)
+					{
+						string Filename;
+						if(AddedFiles.ContainsKey(Product.Path.FullName))
+						{
+							continue;
+						}
+
+						if (LocalRoot != null && Product.Path.IsUnderDirectory(LocalRoot))
+						{
+							Filename = Product.Path.MakeRelativeTo(LocalRoot.ParentDirectory!);
+						}
+						else if(Product.Path.IsUnderDirectory(EngineRoot))
+						{
+							Filename = Product.Path.MakeRelativeTo(EngineRoot);
+						}
+						else
+						{
+							throw new BuildException("Failed to parse target receipt file.  See log for details.");
+						}
+
+						AddedFiles.Add(Product.Path.FullName, Filename);
+					}
+				}
+
+				foreach(var Dep in Receipt.RuntimeDependencies)
+				{
+					if(Dep.Type == StagedFileType.NonUFS)
+					{
+						if(AddedFiles.ContainsKey(Dep.Path.FullName))
+						{
+							continue;
+						}
+
+						string Filename;
+						if (LocalRoot != null && Dep.Path.IsUnderDirectory(LocalRoot))
+						{
+							Filename = Dep.Path.MakeRelativeTo(LocalRoot.ParentDirectory!);
+						}
+						else if (Dep.Path.IsUnderDirectory(EngineRoot))
+						{
+							Filename = Dep.Path.MakeRelativeTo(EngineRoot);
+						}
+						else
+						{
+							throw new BuildException("Failed to parse target receipt file.  See log for details.");
+						}
+
+						AddedFiles.Add(Dep.Path.FullName, Filename);
+					}
+				}
+			}
+
+			string ManifestName = String.Format("AppxManifest_{0}.xml", WindowsExports.GetArchitectureSubpath(Architecture));
+			AddedFiles.Add(Path.Combine(OutputDirectory, ManifestName), "AppxManifest.xml");
+
+			//manually add resources
+			string PriFileName = String.Format("resources_{0}.pri", WindowsExports.GetArchitectureSubpath(Architecture));
+			AddedFiles.Add(Path.Combine(OutputDirectory, PriFileName), "resources.pri");
+			{
+				DirectoryReference ResourceFolder = DirectoryReference.Combine(Receipt.Launch.Directory, WindowsExports.GetArchitectureSubpath(Architecture));
+				foreach (var ResourcePath in UpdatedFiles)
+				{
+					var ResourceFile = new FileReference(ResourcePath);
+
+					if (ResourceFile.IsUnderDirectory(ResourceFolder))
+					{
+						AddedFiles.Add(ResourceFile.FullName, ResourceFile.MakeRelativeTo(ResourceFolder));
+					}
+					else
+					{
+						Log.TraceError("Wrong path to resource \'{0}\', the resource should be in \'{1}\'", ResourceFile.FullName, ResourceFolder.FullName);
+						throw new BuildException("Failed to generate AppX file.  See log for details.");
+					}
+				}
+			}
+
+
+			FileReference SourceNetworkManifestPath = new FileReference(Path.Combine(OutputDirectory, "NetworkManifest.xml"));
+			if (FileReference.Exists(SourceNetworkManifestPath))
+			{
+				AddedFiles.Add(SourceNetworkManifestPath.FullName, "NetworkManifest.xml");
+			}
+			FileReference SourceXboxConfigPath = new FileReference(Path.Combine(OutputDirectory, "xboxservices.config"));
+			if (FileReference.Exists(SourceXboxConfigPath))
+			{
+				AddedFiles.Add(SourceXboxConfigPath.FullName, "xboxservices.config");
+			}
+
+			do
+			{
+				if (PackageFileNeedToBeUpdated)
+				{
+					break;
+				}
+
+				if (!File.Exists(MapFilename))
+				{
+					PackageFileNeedToBeUpdated = true;
+					break;
+				}
+				string[] lines = File.ReadAllLines(MapFilename, Encoding.UTF8);
+				int filesCount = 0;
+
+				foreach(var line in lines)
+				{
+					if (line[0] == '[')
+					{
+						continue;
+					}
+
+					string[] files = line.Split('\t');
+
+					if(files.Length != 2)
+					{
+						PackageFileNeedToBeUpdated = true;
+						break;
+					}
+
+					files[0] = files[0].Trim('\"');
+					files[1] = files[1].Trim('\"');
+
+					if (!AddedFiles.ContainsKey(files[0]))
+					{
+						PackageFileNeedToBeUpdated = true;
+						break;
+					}
+
+					if(AddedFiles[files[0]] != files[1])
+					{
+						PackageFileNeedToBeUpdated = true;
+						break;
+					}
+
+					if(File.GetLastWriteTimeUtc(files[0]).CompareTo(AppXTime) >= 0)
+					{
+						PackageFileNeedToBeUpdated = true;
+						break;
+					}
+
+					++filesCount;
+				}
+
+				if(PackageFileNeedToBeUpdated)
+				{
+					break;
+				}
+
+				if(filesCount != AddedFiles.Count)
+				{
+					PackageFileNeedToBeUpdated = true;
+					break;
+				}
+
+				if (File.Exists(SigningCertificatePath) && File.GetLastWriteTimeUtc(SigningCertificatePath).CompareTo(AppXTime) >= 0)
+				{
+					PackageFileNeedToBeUpdated = true;
+					break;
+				}
+			}
+			while(false);
+
+			if(!PackageFileNeedToBeUpdated)
+			{
+				NewReceipt.BuildProducts.Add(new BuildProduct(new FileReference(OutputAppX), BuildProductType.Package));
+				return;
+			}
+
+			try
+			{
+				DeployHelper_DeleteFile(OutputAppX);
+			}
+			catch (Exception exceptionMessage)
+			{
+				Log.TraceError("Failed to delete {0} from deployment: {1}", OutputAppX, exceptionMessage);
+				throw new BuildException("Failed to generate AppX file.  See log for details.");
+			}
+
+			var AppXRecipeBuiltFiles = new StringBuilder();
+			AppXRecipeBuiltFiles.AppendLine(@"[Files]");
+			foreach (var f in AddedFiles)
+			{
+				AppXRecipeBuiltFiles.AppendLine(String.Format("\"{0}\"\t\"{1}\"", f.Key, f.Value));
+			}
+			File.WriteAllText(MapFilename, AppXRecipeBuiltFiles.ToString(), Encoding.UTF8);
+
+			string MakeAppXCommandLine = String.Format("pack /o /f \"{0}\" /p \"{1}\"", MapFilename, OutputAppX);
+
+			var StartInfo = new ProcessStartInfo(MakeAppXPath!.FullName, MakeAppXCommandLine);
+			StartInfo.UseShellExecute = false;
+			StartInfo.CreateNoWindow = true;
+			var ExitCode = Utils.RunLocalProcessAndPrintfOutput(StartInfo);
+			if (ExitCode < 0)
+			{
+				throw new BuildException("Failed to generate AppX file.  See log for details.");
+			}
+
+			if (File.Exists(SigningCertificatePath))
+			{
+				string SignToolCommandLine = String.Format("sign /a /f \"{0}\" /fd SHA256 \"{1}\"", SigningCertificatePath, OutputAppX);
+				StartInfo = new ProcessStartInfo(SignToolPath!.FullName, SignToolCommandLine);
+				StartInfo.UseShellExecute = false;
+				StartInfo.CreateNoWindow = true;
+				ExitCode = Utils.RunLocalProcessAndPrintfOutput(StartInfo);
+				if (ExitCode < 0)
+				{
+					throw new BuildException("Failed to generate AppX file.  See log for details.");
+				}
+			}
+
+			Log.TraceInformation("AppX successfully packaged to \'{0}{1}\'", OutputName, Extension);
+			NewReceipt.BuildProducts.Add(new BuildProduct(new FileReference(OutputAppX), BuildProductType.Package));
+		}
+
+
+		private void CopyDataAndSymbolsBetweenReceipts(TargetReceipt Receipt, TargetReceipt NewReceipt, WindowsArchitecture Architecture)
+		{
+			NewReceipt.AdditionalProperties.AddRange(Receipt.AdditionalProperties);
+			NewReceipt.AdditionalProperties = NewReceipt.AdditionalProperties.GroupBy(e => e.Name).Select(g => g.First()).ToList();
+
+			NewReceipt.BuildProducts.AddRange(Receipt.BuildProducts.FindAll(x => x.Type == BuildProductType.SymbolFile));
+			NewReceipt.BuildProducts = NewReceipt.BuildProducts.GroupBy(e => e.Path).Select(g => g.First()).ToList();
+
+			NewReceipt.RuntimeDependencies.AddRange(Receipt.RuntimeDependencies.FindAll(x => x.Type != StagedFileType.NonUFS));
+			NewReceipt.RuntimeDependencies = new RuntimeDependencyList(NewReceipt.RuntimeDependencies.GroupBy(e => e.Path).Select(g => g.First()).ToList());
+		}
+
+
+		public override bool PrepTargetForDeployment(TargetReceipt Receipt)
+		{
+			// Use the project name if possible - InTarget.AppName changes for 'Client'/'Server' builds
+			string ProjectName = Receipt.ProjectFile != null ? Receipt.ProjectFile.GetFileNameWithoutAnyExtensions() : Receipt.Launch!.GetFileNameWithoutExtension();
+			Log.TraceInformation("Prepping {0} for deployment to {1}", ProjectName, Receipt.Platform.ToString());
+			System.DateTime PrepDeployStartTime = DateTime.UtcNow;
+
+			// Note: TargetReceipt.Read now expands path variables internally.
+			TargetReceipt? NewReceipt = null;
+			FileReference ReceiptFileName = TargetReceipt.GetDefaultPath(Receipt.ProjectDir != null ? Receipt.ProjectDir : Unreal.EngineDirectory, Receipt.TargetName, Receipt.Platform, Receipt.Configuration, "Multi");
+			if (!TargetReceipt.TryRead(ReceiptFileName, Unreal.EngineDirectory, out NewReceipt))
+			{
+				NewReceipt = new TargetReceipt(Receipt.ProjectFile, Receipt.TargetName, Receipt.TargetType, Receipt.Platform, Receipt.Configuration, Receipt.Version, "Multi", Receipt.IsTestTarget);
+			}
+
+			string SDK = "";
+			var Results = Receipt.AdditionalProperties.Where(x => x.Name == "SDK");
+			if (Results.Any())
+			{
+				SDK = Results.First().Value;
+			}
+			HoloLensExports.InitWindowsSdkToolPath(SDK);
+
+			AddWinMDReferencesFromReceipt(Receipt, Receipt.ProjectDir != null ? Receipt.ProjectDir : Unreal.EngineDirectory, Unreal.EngineDirectory.ParentDirectory!.FullName, SDK);
+
+			//PrepForUATPackageOrDeploy(InTarget.ProjectFile, InAppName, InTarget.ProjectDirectory.FullName, InTarget.OutputPath.FullName, TargetBuildEnvironment.RelativeEnginePath, false, "", false);
+			List<UnrealTargetConfiguration> TargetConfigs = new List<UnrealTargetConfiguration> { Receipt.Configuration };
+			List<string> ExePaths = new List<string> { Receipt.Launch!.FullName };
+			string RelativeEnginePath = Unreal.EngineDirectory.MakeRelativeTo(DirectoryReference.GetCurrentDirectory());
+
+			WindowsArchitecture Arch = WindowsArchitecture.ARM64;
+			if (Receipt.Architecture.ToLower() == "x64")
+			{
+				Arch = WindowsArchitecture.x64;
+			}
+
+
+			string AbsoluteExeDirectory = Path.GetDirectoryName(ExePaths[0])!;
+			UnrealTargetPlatform Platform = UnrealTargetPlatform.HoloLens;
+			string IntermediateDirectory = Path.Combine(Receipt.ProjectDir != null ? Receipt.ProjectDir.FullName : Unreal.EngineDirectory.FullName, "Intermediate", "Deploy", WindowsExports.GetArchitectureSubpath(Arch));
+			List<string> UpdatedFiles = new HoloLensManifestGenerator().CreateManifest(Platform, Arch, AbsoluteExeDirectory, IntermediateDirectory, Receipt.ProjectFile, Receipt.ProjectDir != null ? Receipt.ProjectDir.FullName : Unreal.EngineDirectory.FullName, TargetConfigs, ExePaths, WinMDReferences)!;
+
+			PrepForUATPackageOrDeploy(Receipt.ProjectFile, ProjectName, Receipt.ProjectDir != null ? Receipt.ProjectDir.FullName : Unreal.EngineDirectory.FullName, Arch, TargetConfigs, ExePaths, RelativeEnginePath, false, "", false);
+			MakePackage(Receipt, NewReceipt, Arch, UpdatedFiles);
+			CopyDataAndSymbolsBetweenReceipts(Receipt, NewReceipt, Arch);
+
+			NewReceipt.Write(ReceiptFileName, Unreal.EngineDirectory);
+
+			// Log out the time taken to deploy...
+			double PrepDeployDuration = (DateTime.UtcNow - PrepDeployStartTime).TotalSeconds;
+			Log.TraceInformation("HoloLens deployment preparation took {0:0.00} seconds", PrepDeployDuration);
+
+			return true;
+		}
+
+		public void AddWinMDReferencesFromReceipt(TargetReceipt Receipt, DirectoryReference SourceProjectDir, string DestRelativeTo, string SDKVersion)
+		{
+			// Dependency paths in receipt are already expanded at this point
+			foreach (var Dep in Receipt.RuntimeDependencies)
+			{
+				if (Dep.Path.GetExtension() == ".dll")
+				{
+					string SourcePath = Dep.Path.FullName;
+					string WinMDFile = Path.ChangeExtension(SourcePath, "winmd");
+					if (File.Exists(WinMDFile))
+					{
+						string DestPath = Dep.Path.FullName;
+						DestPath = Dep.Path.FullName.Replace(Unreal.EngineDirectory.FullName, Path.Combine(DestRelativeTo, "Engine"));
+						DestPath = DestPath.Replace(SourceProjectDir.FullName, Path.Combine(DestRelativeTo, SourceProjectDir.GetDirectoryName()));
+						DestPath = Utils.MakePathRelativeTo(DestPath, DestRelativeTo);
+						WinMDReferences.Add(new WinMDRegistrationInfo(new FileReference(WinMDFile), DestPath, SDKVersion));
+					}
+				}
+			}
+		}
+
+		private List<WinMDRegistrationInfo> WinMDReferences = new List<WinMDRegistrationInfo>();
+	}
+}
diff --git a/Engine/Platforms/Hololens/Source/Programs/UnrealBuildTool/HoloLensExports.cs b/Engine/Platforms/Hololens/Source/Programs/UnrealBuildTool/HoloLensExports.cs
new file mode 100644
--- /dev/null
+++ b/Engine/Platforms/Hololens/Source/Programs/UnrealBuildTool/HoloLensExports.cs
@@ -0,0 +1,100 @@
+// Copyright Epic Games, Inc. All Rights Reserved.
+
+using System;
+using System.Collections.Generic;
+using System.Linq;
+using System.Runtime.Versioning;
+using System.Text;
+using System.Threading.Tasks;
+using EpicGames.Core;
+
+namespace UnrealBuildTool
+{
+	/// <summary>
+	/// Public HoloLensDeploy wrapper exposed to UAT
+	/// </summary>
+	[SupportedOSPlatform("windows")]
+	public class HoloLensExports
+	{
+		private HoloLensDeploy InnerDeploy;
+
+		/// <summary>
+		/// 
+		/// </summary>
+		public HoloLensExports()
+		{
+			InnerDeploy = new HoloLensDeploy();
+		}
+
+		/// <summary>
+		/// 
+		/// </summary>
+		/// <param name="ProjectFile"></param>
+		/// <param name="InProjectName"></param>
+		/// <param name="InProjectDirectory"></param>
+		/// <param name="Architecture"></param>
+		/// <param name="InTargetConfigurations"></param>
+		/// <param name="InExecutablePaths"></param>
+		/// <param name="InEngineDir"></param>
+		/// <param name="bForDistribution"></param>
+		/// <param name="CookFlavor"></param>
+		/// <param name="bIsDataDeploy"></param>
+		/// <returns></returns>
+		public bool PrepForUATPackageOrDeploy(FileReference ProjectFile, string InProjectName, string InProjectDirectory, WindowsArchitecture Architecture, List<UnrealTargetConfiguration> InTargetConfigurations, List<string> InExecutablePaths, string InEngineDir, bool bForDistribution, string CookFlavor, bool bIsDataDeploy)
+		{
+			return InnerDeploy.PrepForUATPackageOrDeploy(ProjectFile, InProjectName, InProjectDirectory, Architecture, InTargetConfigurations, InExecutablePaths, InEngineDir, bForDistribution, CookFlavor, bIsDataDeploy);
+		}
+
+		/// <summary>
+		/// Collect all the WinMD references
+		/// </summary>
+		/// <param name="Receipt"></param>
+		/// <param name="SourceProjectDir"></param>
+		/// <param name="DestPackageRoot"></param>
+		/// <param name="SdkVersion"></param>
+		public void AddWinMDReferencesFromReceipt(TargetReceipt Receipt, DirectoryReference SourceProjectDir, string DestPackageRoot, string SdkVersion)
+		{
+			InnerDeploy.AddWinMDReferencesFromReceipt(Receipt, SourceProjectDir, DestPackageRoot, SdkVersion);
+		}
+
+		/// <summary>
+		/// 
+		/// </summary>
+		/// <param name="ToolName"></param>
+		/// <returns></returns>
+		public static FileReference? GetWindowsSdkToolPath(string ToolName)
+		{
+			return HoloLensToolChain.GetWindowsSdkToolPath(ToolName);
+		}
+
+		/// <summary>
+		/// 
+		/// </summary>
+		/// <param name="SdkVersion"></param>
+		/// <returns></returns>
+		public static bool InitWindowsSdkToolPath(string SdkVersion)
+		{
+			return HoloLensToolChain.InitWindowsSdkToolPath(SdkVersion);
+		}
+
+		/// <summary>
+		/// 
+		/// </summary>
+		/// <param name="DLCFile"></param>
+		/// <param name="OutputDirectory"></param>
+		public static void CreateManifestForDLC(FileReference DLCFile, DirectoryReference OutputDirectory)
+		{
+			string IntermediateDirectory = DirectoryReference.Combine(DLCFile.Directory, "Intermediate", "Deploy").FullName;
+			new HoloLensManifestGenerator().CreateManifest(UnrealTargetPlatform.HoloLens, WindowsArchitecture.ARM64, OutputDirectory.FullName, IntermediateDirectory, DLCFile, DLCFile.Directory.FullName, new List<UnrealTargetConfiguration>(), new List<string>(), null);
+		}
+
+		/// <summary>
+		/// 
+		/// </summary>
+		/// <returns>CurrentWindowsSdkVersion</returns>
+		public static Version GetCurrentWindowsSdkVersion()
+		{
+			return HoloLensToolChain.GetCurrentWindowsSdkVersion();
+		}
+	}
+}
diff --git a/Engine/Platforms/Hololens/Source/Programs/UnrealBuildTool/HoloLensProjectGenerator.cs b/Engine/Platforms/Hololens/Source/Programs/UnrealBuildTool/HoloLensProjectGenerator.cs
new file mode 100644
--- /dev/null
+++ b/Engine/Platforms/Hololens/Source/Programs/UnrealBuildTool/HoloLensProjectGenerator.cs
@@ -0,0 +1,219 @@
+// Copyright Epic Games, Inc. All Rights Reserved.
+
+using System;
+using System.Collections.Generic;
+using System.Text;
+using System.Diagnostics;
+using System.IO;
+using System.Security;
+using System.Xml.Linq;
+using EpicGames.Core;
+using UnrealBuildBase;
+using System.Runtime.Versioning;
+
+namespace UnrealBuildTool
+{
+	/// <summary>
+	/// Base class for platform-specific project generators
+	/// </summary>
+	[SupportedOSPlatform("windows")]
+	class HoloLensProjectGenerator : PlatformProjectGenerator
+	{
+		const string PlatformString = "HoloLens";
+
+		/// <summary>
+		/// Constructor
+		/// </summary>
+		/// <param name="Arguments">Command line arguments passed to the project generator</param>
+		public HoloLensProjectGenerator(CommandLineArguments Arguments)
+			: base(Arguments)
+		{
+		}
+
+		/// <summary>
+		/// Enumerate all the platforms that this generator supports
+		/// </summary>
+		public override IEnumerable<UnrealTargetPlatform> GetPlatforms()
+		{
+			yield return UnrealTargetPlatform.HoloLens;
+		}
+
+		///
+		///	VisualStudio project generation functions
+		///	
+		/// <summary>
+		/// Whether this build platform has native support for VisualStudio
+		/// </summary>
+		/// <param name="InPlatform">  The UnrealTargetPlatform being built</param>
+		/// <param name="InConfiguration"> The UnrealTargetConfiguration being built</param>
+		/// <param name="ProjectFileFormat"></param>
+		/// <returns>bool    true if native VisualStudio support (or custom VSI) is available</returns>
+		public override bool HasVisualStudioSupport(UnrealTargetPlatform InPlatform, UnrealTargetConfiguration InConfiguration, VCProjectFileFormat ProjectFileFormat)
+		{
+			return true;
+		}
+
+		/// <summary>
+		/// Get whether this platform deploys
+		/// </summary>
+		/// <returns>bool  true if the 'Deploy' option should be enabled</returns>
+		public override bool GetVisualStudioDeploymentEnabled(UnrealTargetPlatform InPlatform, UnrealTargetConfiguration InConfiguration)
+		{
+			return true;
+		}
+
+		public override void GenerateGameProperties(UnrealTargetConfiguration Configuration, StringBuilder VCProjectFileContent, TargetType TargetType, DirectoryReference RootDirectory, FileReference TargetFilePath)
+		{
+			string MinVersion = string.Empty;
+			string MaxTestedVersion = string.Empty;
+			ConfigHierarchy EngineIni = ConfigCache.ReadHierarchy(ConfigHierarchyType.Engine, RootDirectory, UnrealTargetPlatform.HoloLens);
+			if (EngineIni != null)
+			{
+				EngineIni.GetString("/Script/HoloLensPlatformEditor.HoloLensTargetSettings", "MinimumPlatformVersion", out MinVersion);
+				EngineIni.GetString("/Script/HoloLensPlatformEditor.HoloLensTargetSettings", "MaximumPlatformVersionTested", out MaxTestedVersion);
+			}
+			if (!string.IsNullOrEmpty(MinVersion))
+			{
+				VCProjectFileContent.Append("		<WindowsTargetPlatformMinVersion>" + MinVersion + "</WindowsTargetPlatformMinVersion>" + ProjectFileGenerator.NewLine);
+			}
+			if (!string.IsNullOrEmpty(MaxTestedVersion))
+			{
+				VCProjectFileContent.Append("		<WindowsTargetPlatformVersion>" + MaxTestedVersion + "</WindowsTargetPlatformVersion>" + ProjectFileGenerator.NewLine);
+			}
+
+			WindowsCompiler Compiler = WindowsPlatform.GetDefaultCompiler(TargetFilePath, WindowsArchitecture.x64);
+			DirectoryReference? PlatformWinMDLocation = HoloLensPlatform.GetCppCXMetadataLocation(Compiler, "Latest", WindowsArchitecture.x64);
+			if (PlatformWinMDLocation == null || !FileReference.Exists(FileReference.Combine(PlatformWinMDLocation, "platform.winmd")))
+			{
+				throw new BuildException("Unable to find platform.winmd for {0} toolchain; '{1}' is an invalid version", WindowsPlatform.GetCompilerName(Compiler), "Latest");
+			}
+			string FoundationWinMDPath = HoloLensPlatform.GetLatestMetadataPathForApiContract("Windows.Foundation.FoundationContract", Compiler);
+			string UniversalWinMDPath = HoloLensPlatform.GetLatestMetadataPathForApiContract("Windows.Foundation.UniversalApiContract", Compiler);
+			VCProjectFileContent.Append("		<AdditionalOptions>/ZW /ZW:nostdlib</AdditionalOptions>" + ProjectFileGenerator.NewLine);
+			VCProjectFileContent.Append("		<NMakePreprocessorDefinitions>$(NMakePreprocessorDefinitions);PLATFORM_HOLOLENS=1;HOLOLENS=1;</NMakePreprocessorDefinitions>" + ProjectFileGenerator.NewLine);
+			if (PlatformWinMDLocation != null)
+			{
+				VCProjectFileContent.Append("		<NMakeAssemblySearchPath>$(NMakeAssemblySearchPath);" + PlatformWinMDLocation + "</NMakeAssemblySearchPath>" + ProjectFileGenerator.NewLine);
+			}
+			VCProjectFileContent.Append("		<NMakeForcedUsingAssemblies>$(NMakeForcedUsingAssemblies);" + FoundationWinMDPath + ";" + UniversalWinMDPath + ";platform.winmd</NMakeForcedUsingAssemblies>" + ProjectFileGenerator.NewLine);
+		}
+
+		public override void GetVisualStudioPreDefaultString(UnrealTargetPlatform InPlatform, UnrealTargetConfiguration InConfiguration, StringBuilder VCProjectFileContent)
+		{
+			// VS2017 expects WindowsTargetPlatformVersion to be set in conjunction with these other properties, otherwise the projects
+			// will fail to load when the solution is in a HoloLens configuration.
+			// Default to latest supported version.  Game projects can override this later.
+			// Because this property is only required for VS2017 we can safely say that's the compiler version (whether that's actually true
+			// or not)
+			//string SDKFolder = "";
+			string SDKVersion = "";
+
+			DirectoryReference? folder;
+			VersionNumber? version;
+			if(WindowsPlatform.TryGetWindowsSdkDir("Latest", out version, out folder))
+			{
+				//SDKFolder = folder.FullName;
+				SDKVersion = version.ToString();
+			}
+
+			VCProjectFileContent.AppendLine("		<AppContainerApplication>true</AppContainerApplication>");
+			VCProjectFileContent.AppendLine("		<ApplicationType>Windows Store</ApplicationType>");
+			VCProjectFileContent.AppendLine("		<ApplicationTypeRevision>10.0</ApplicationTypeRevision>");
+			VCProjectFileContent.AppendLine("		<WindowsAppContainer>true</WindowsAppContainer>");
+			VCProjectFileContent.AppendLine("		<AppxPackage>true</AppxPackage>");
+			VCProjectFileContent.AppendLine("		<WindowsTargetPlatformVersion>{0}</WindowsTargetPlatformVersion>", SDKVersion.ToString()); 
+		}
+
+		public override string GetVisualStudioLayoutDirSection(UnrealTargetPlatform InPlatform, UnrealTargetConfiguration InConfiguration, string InConditionString, TargetType TargetType, FileReference TargetRulesPath, FileReference ProjectFilePath, FileReference NMakeOutputPath, VCProjectFileFormat InProjectFileFormat)
+		{
+			string LayoutDirString = "";
+
+			if (IsValidHoloLensTarget(InPlatform, TargetType, TargetRulesPath))
+			{
+				LayoutDirString += "	<PropertyGroup " + InConditionString + ">" + ProjectFileGenerator.NewLine;
+				LayoutDirString += "		<RemoveExtraDeployFiles>false</RemoveExtraDeployFiles>" + ProjectFileGenerator.NewLine;
+				LayoutDirString += "		<LayoutDir>" + SecurityElement.Escape(DirectoryReference.Combine(NMakeOutputPath.Directory, "AppX").FullName) + "</LayoutDir>" + ProjectFileGenerator.NewLine;
+				LayoutDirString += "		<AppxPackageRecipe>" + SecurityElement.Escape(FileReference.Combine(NMakeOutputPath.Directory, ProjectFilePath.GetFileNameWithoutExtension() + ".build.appxrecipe").FullName) + "</AppxPackageRecipe>" + ProjectFileGenerator.NewLine;
+				LayoutDirString += "	</PropertyGroup>" + ProjectFileGenerator.NewLine;
+
+				// another hijack - this is a convenient point to make sure that HoloLens-appropriate debuggers are available
+				// in the project property pages.
+				LayoutDirString += "    <ItemGroup " + InConditionString + ">" + ProjectFileGenerator.NewLine;
+				LayoutDirString += "		<PropertyPageSchema Include=\"$(VCTargetsPath)$(LangID)\\AppHostDebugger_Local.xml\" />" + ProjectFileGenerator.NewLine;
+				LayoutDirString += "		<PropertyPageSchema Include=\"$(VCTargetsPath)$(LangID)\\AppHostDebugger_Simulator.xml\" />" + ProjectFileGenerator.NewLine;
+				LayoutDirString += "		<PropertyPageSchema Include=\"$(VCTargetsPath)$(LangID)\\AppHostDebugger_Remote.xml\" />" + ProjectFileGenerator.NewLine;
+				LayoutDirString += "    </ItemGroup>" + ProjectFileGenerator.NewLine;
+			}
+
+			return LayoutDirString;
+		}
+
+		///
+		///	VisualStudio project generation functions
+		///	
+		/// <summary>
+		/// Return the VisualStudio platform name for this build platform
+		/// </summary>
+		/// <param name="InPlatform">  The UnrealTargetPlatform being built</param>
+		/// <param name="InConfiguration"> The UnrealTargetConfiguration being built</param>
+		/// <returns>string    The name of the platform that VisualStudio recognizes</returns>
+		public override string GetVisualStudioPlatformName(UnrealTargetPlatform InPlatform, UnrealTargetConfiguration InConfiguration)
+		{
+			if (InPlatform == UnrealTargetPlatform.HoloLens)
+			{
+				return "arm64";
+			}
+			return InPlatform.ToString();
+		}
+
+		private bool IsValidHoloLensTarget(UnrealTargetPlatform InPlatform, TargetType InTargetType, FileReference InTargetFilePath)
+		{
+			if ((InPlatform == UnrealTargetPlatform.HoloLens) &&
+				(InTargetType == TargetRules.TargetType.Client || InTargetType == TargetRules.TargetType.Game ) &&
+				InTargetType != TargetRules.TargetType.Editor &&
+				InTargetType != TargetRules.TargetType.Server
+				)
+			{
+				// We do not want to include any Templates targets
+				// Not a huge fan of doing it via path name comparisons... but it works
+				string TempTargetFilePath = InTargetFilePath.FullName.Replace("\\", "/");
+				if (TempTargetFilePath.Contains("/Templates/"))
+				{
+					string AbsoluteEnginePath = Unreal.EngineDirectory.FullName;
+					AbsoluteEnginePath = AbsoluteEnginePath.Replace("\\", "/");
+					if (AbsoluteEnginePath.EndsWith("/") == false)
+					{
+						AbsoluteEnginePath += "/";
+					}
+					string CheckPath = AbsoluteEnginePath.Replace("/Engine/", "/Templates/");
+					if (TempTargetFilePath.StartsWith(CheckPath))
+					{
+						return false;
+					}
+				}
+				return true;
+			}
+
+			return false;
+		}
+
+		public override bool RequiresVSUserFileGeneration()
+		{
+			return true;
+		}
+
+		public override string GetVisualStudioUserFileStrings(UnrealTargetPlatform InPlatform, UnrealTargetConfiguration InConfiguration,
+			string InConditionString, TargetRules InTargetRules, FileReference TargetRulesPath, FileReference ProjectFilePath)
+		{
+			string UserFileEntry = "";
+			if (IsValidHoloLensTarget(InPlatform, InTargetRules.Type, TargetRulesPath))
+			{
+				UserFileEntry += "<PropertyGroup " + InConditionString + ">\n";
+				UserFileEntry += "	<DebuggerFlavor>AppHostLocalDebugger</DebuggerFlavor>\n";
+				UserFileEntry += "</PropertyGroup>\n";
+			}
+			return UserFileEntry;
+		}
+	}
+}
+
diff --git a/Engine/Platforms/Hololens/Source/Programs/UnrealBuildTool/HoloLensToolChain.cs b/Engine/Platforms/Hololens/Source/Programs/UnrealBuildTool/HoloLensToolChain.cs
new file mode 100644
--- /dev/null
+++ b/Engine/Platforms/Hololens/Source/Programs/UnrealBuildTool/HoloLensToolChain.cs
@@ -0,0 +1,206 @@
+// Copyright Epic Games, Inc. All Rights Reserved.
+
+using System;
+using System.Collections.Generic;
+using System.Diagnostics;
+using System.IO;
+using System.Linq;
+using System.Text;
+using System.Text.RegularExpressions;
+using Microsoft.Win32;
+using EpicGames.Core;
+using UnrealBuildBase;
+using System.Runtime.Versioning;
+
+namespace UnrealBuildTool
+{
+	static class Extension
+	{
+		public static void AddFormat(this List<string> list, string formatString, params object[] args)
+		{
+			list.Add(String.Format(formatString, args));
+		}
+	}
+
+	[SupportedOSPlatform("windows")]
+	class HoloLensToolChain : VCToolChain
+	{
+		public HoloLensToolChain(ReadOnlyTargetRules Target)
+			: base(Target)
+		{
+			// by default tools chains don't parse arguments, but we want to be able to check the -architectures flag defined above. This is
+			// only necessary when AndroidToolChain is used during UAT
+			CommandLine.ParseArguments(Environment.GetCommandLineArgs(), this);
+		}
+
+		protected override void AppendCLArguments_Global(CppCompileEnvironment CompileEnvironment, List<string> Arguments)
+		{
+			base.AppendCLArguments_Global(CompileEnvironment, Arguments);
+
+			// @todo HoloLens: Disable "unreachable code" warning since auto-included vccorlib.h triggers it
+			Arguments.Add("/wd4702");
+
+			// @todo HoloLens: Disable "usage of ATL attributes is deprecated" since WRL headers generate this
+			Arguments.Add("/wd4467");
+
+			// @todo HoloLens: Disable "declaration of 'identifier' hides class member"
+			Arguments.Add("/wd4458");
+
+			// @todo HoloLens: Disable "not defined as a preprocessor macro, replacing with '0'"
+			Arguments.Add("/wd4668");
+
+			Arguments.RemoveAll(Argument => Argument.StartsWith("/RTC"));
+			if (CompileEnvironment.Configuration == CppConfiguration.Debug && CompileEnvironment.bUseDebugCRT)
+			{
+				Arguments.Add("/RTC1");
+			}
+
+			// C2338: two-phase name lookup is not supported for C++/CLI or C++/CX; use /Zc:twoPhase-
+			if (Target.WindowsPlatform.Compiler.IsMSVC() && Target.WindowsPlatform.bStrictConformanceMode)
+			{
+				Arguments.Add("/Zc:twoPhase-");
+			}
+
+			DirectoryReference? PlatformWinMDLocation = HoloLensPlatform.GetCppCXMetadataLocation(EnvVars.Compiler, EnvVars.ToolChainDir);
+			if (PlatformWinMDLocation != null)
+			{
+				Arguments.AddFormat(@"/AI""{0}""", PlatformWinMDLocation);
+				Arguments.AddFormat(@"/FU""{0}\platform.winmd""", PlatformWinMDLocation);
+			}
+		}
+
+		protected override void AppendCLArguments_CPP(CppCompileEnvironment CompileEnvironment, List<string> Arguments)
+		{
+			base.AppendCLArguments_CPP(CompileEnvironment, Arguments);
+
+			// Enable Windows Runtime extensions.  Do this even for libs (plugins) so that these too can consume WinRT APIs
+			Arguments.Add("/ZW");
+
+			// Don't automatically add metadata references.  We'll do that ourselves to avoid referencing windows.winmd directly:
+			// we've hit problems where types are somehow in windows.winmd on some installations but not others, leading to either
+			// missing or duplicated type references.
+			Arguments.Add("/ZW:nostdlib");
+		}
+
+		protected override void AppendLinkArguments(LinkEnvironment LinkEnvironment, List<string> Arguments)
+		{
+			base.AppendLinkArguments(LinkEnvironment, Arguments);
+
+			// @todo HoloLens: This is logic from before inheriting from VCToolChain. Is not embedding the manifest necessary?
+			if (!Arguments.Contains("/MANIFEST:NO"))
+			{
+				Arguments.Remove("/MANIFEST:EMBED");
+				Arguments.RemoveAll((x) => x.StartsWith("/MANIFESTINPUT"));
+
+				// Don't create a side-by-side manifest file for the executable.
+				Arguments.Add("/MANIFEST:NO");
+			}
+
+			Arguments.Add("/APPCONTAINER");
+
+			// this helps with store API compliance validation tools, adding additional pdb info
+			Arguments.Add("/PROFILE");
+		}
+
+		protected override void AppendLibArguments(LinkEnvironment LinkEnvironment, List<string> Arguments)
+		{
+			base.AppendLibArguments(LinkEnvironment, Arguments);
+
+			// Ignore warning about /ZW in static libraries. It's not relevant since UE modules
+			// have no reason to export new WinRT types, and ignoring it quiets noise when using
+			// WinRT APIs from plugins.
+			Arguments.Add("/ignore:4264");
+		}
+
+		private static DirectoryReference? CurrentWindowsSdkBinDir = null;
+		private static Version? CurrentWindowsSdkVersion;
+
+		public static bool InitWindowsSdkToolPath(string SdkVersion)
+		{
+			if (string.IsNullOrEmpty(SdkVersion))
+			{
+				Log.TraceError("WinSDK version is empty");
+				return false;
+			}
+
+			VersionNumber OutSdkVersion;
+			DirectoryReference OutSdkDir;
+
+			if (!WindowsPlatform.TryGetWindowsSdkDir(SdkVersion, out OutSdkVersion!, out OutSdkDir!))
+			{
+				Log.TraceError("Failed to find WinSDK " + SdkVersion);
+				return false;
+			}
+
+			DirectoryReference WindowsSdkBinDir = DirectoryReference.Combine(OutSdkDir, "bin", OutSdkVersion.ToString(), Environment.Is64BitProcess ? "x64" : "x86");
+
+			if(!DirectoryReference.Exists(WindowsSdkBinDir))
+			{
+				Log.TraceError("WinSDK " + SdkVersion + " doesn't exit");
+				return false;
+			}
+
+			CurrentWindowsSdkVersion = new Version(OutSdkVersion.ToString());
+			CurrentWindowsSdkBinDir = WindowsSdkBinDir;
+			return true;
+		}
+
+		public static FileReference? GetWindowsSdkToolPath(string ToolName)
+		{
+			FileReference file = FileReference.Combine(CurrentWindowsSdkBinDir!, ToolName);
+
+			if(!FileReference.Exists(file))
+			{
+				return null;
+			}
+
+			return file;
+		}
+
+		public static Version GetCurrentWindowsSdkVersion()
+		{
+			return CurrentWindowsSdkVersion!;
+		}
+
+		public override string GetSDKVersion()
+		{
+			return CurrentWindowsSdkVersion!.ToString();
+		}
+
+		public override void ModifyBuildProducts(ReadOnlyTargetRules Target, UEBuildBinary Binary, List<string> Libraries, List<UEBuildBundleResource> BundleResources, Dictionary<FileReference, BuildProductType> BuildProducts)
+		{
+			base.ModifyBuildProducts(Target, Binary, Libraries, BundleResources, BuildProducts);
+
+			DirectoryReference HoloLensBinaryDirectory = Binary.OutputFilePath.Directory;
+
+			AddBuildProductSafe(BuildProducts, FileReference.Combine(HoloLensBinaryDirectory, "AppxManifest_"+ Target.Architecture + ".xml"), BuildProductType.BuildResource);
+			AddBuildProductSafe(BuildProducts, FileReference.Combine(HoloLensBinaryDirectory, "resources_"+ Target.Architecture + ".pri"), BuildProductType.BuildResource);
+			if (Target.Configuration == UnrealTargetConfiguration.Development)
+			{
+				AddBuildProductSafe(BuildProducts, FileReference.Combine(HoloLensBinaryDirectory, Target.Name + Target.Architecture + ".exe"), BuildProductType.Executable);
+				AddBuildProductSafe(BuildProducts, FileReference.Combine(HoloLensBinaryDirectory, Target.Name + Target.Architecture + ".pdb"), BuildProductType.SymbolFile);
+			}
+			else
+			{
+				AddBuildProductSafe(BuildProducts, FileReference.Combine(HoloLensBinaryDirectory, Target.Name + "-HoloLens-" + Target.Configuration + Target.Architecture + ".exe"), BuildProductType.Executable);
+				AddBuildProductSafe(BuildProducts, FileReference.Combine(HoloLensBinaryDirectory, Target.Name + "-HoloLens-" + Target.Configuration + Target.Architecture + ".pdb"), BuildProductType.SymbolFile);
+			}
+			AddBuildProductSafe(BuildProducts, FileReference.Combine(HoloLensBinaryDirectory, Target.Name + "-HoloLens-" + Target.Configuration + Target.Architecture + ".target"), BuildProductType.BuildResource);
+			AddBuildProductSafe(BuildProducts, FileReference.Combine(HoloLensBinaryDirectory, Target.Architecture + "\\Resources\\Logo.png"), BuildProductType.BuildResource);
+			AddBuildProductSafe(BuildProducts, FileReference.Combine(HoloLensBinaryDirectory, Target.Architecture + "\\Resources\\resources.resw"), BuildProductType.BuildResource);
+			AddBuildProductSafe(BuildProducts, FileReference.Combine(HoloLensBinaryDirectory, Target.Architecture + "\\Resources\\SmallLogo.png"), BuildProductType.BuildResource);
+			AddBuildProductSafe(BuildProducts, FileReference.Combine(HoloLensBinaryDirectory, Target.Architecture + "\\Resources\\SplashScreen.png"), BuildProductType.BuildResource);
+			AddBuildProductSafe(BuildProducts, FileReference.Combine(HoloLensBinaryDirectory, Target.Architecture + "\\Resources\\WideLogo.png"), BuildProductType.BuildResource);
+			AddBuildProductSafe(BuildProducts, FileReference.Combine(HoloLensBinaryDirectory, Target.Architecture + "\\Resources\\en\\resources.resw"), BuildProductType.BuildResource);
+		}
+
+		private void AddBuildProductSafe(Dictionary<FileReference, BuildProductType> BuildProducts, FileReference FileToAdd, BuildProductType ProductType)
+		{
+			if (!BuildProducts.ContainsKey(FileToAdd))
+			{
+				BuildProducts.Add(FileToAdd, ProductType);
+			}
+		}
+
+	};
+}
diff --git a/Engine/Platforms/Hololens/Source/Programs/UnrealBuildTool/ManifestGenerator.cs b/Engine/Platforms/Hololens/Source/Programs/UnrealBuildTool/ManifestGenerator.cs
new file mode 100644
--- /dev/null
+++ b/Engine/Platforms/Hololens/Source/Programs/UnrealBuildTool/ManifestGenerator.cs
@@ -0,0 +1,2512 @@
+// Copyright Epic Games, Inc. All Rights Reserved.
+
+using System;
+using System.Collections.Generic;
+using System.Linq;
+using System.IO;
+using System.Text;
+using System.Text.RegularExpressions;
+using System.Resources;
+using System.Xml;
+using EpicGames.Core;
+using UnrealBuildBase;
+using System.Runtime.Versioning;
+
+namespace UnrealBuildTool
+{
+	/// <summary>
+	///  Class to handle generating an AppxManifest.xml file
+	/// </summary>
+	public class HoloLensManifestGenerator
+	{
+		// Global path configuration
+		private string BuildResourceSubPath = "Resources";
+		private string? StoreResourceSubPath;
+		private const string EngineResourceSubPath = "DefaultImages";
+
+		// Manifest compliance values
+		private const int MaxResourceEntries = 200;
+
+		// INI configuration cache
+		private ConfigHierarchy? EngineIni;
+		private ConfigHierarchy? GameIni;
+
+		// Manifest configuration values/paths
+		private List<WinMDRegistrationInfo>? WinMDReferences;
+		//private UnrealTargetPlatform Platform;
+		private WindowsArchitecture Architecture;
+		private string? TargetSettings;
+		private string? ProjectPath;
+		private string? OutputPath;
+		private string? IntermediatePath;
+		private List<string>? CulturesToStage;
+
+		// Manifest generation state
+		private UEResXWriter? NeutralResourceWriter;
+		private List<UEResXWriter>? PerCultureResourceWriters;
+		private XmlDocument? AppxManifestXmlDocument;
+		private List<string>? UpdatedFilePaths;
+		private List<string>? ManifestRelatedFilePaths;
+
+		// Analagous to RelativeProjectRootForStage in UAT so that VS (UBT only) and UAT layouts match
+		private string? RelativeProjectRootForStage;
+		bool IsGameSpecificExe;
+		private bool IsDlc;
+		private Dictionary<string, string>? ParsedDlcInfo;
+
+
+		/// <summary>
+		/// Retrieve a package configuration option from the deprecated [AppxManifest] INI settings and return the value.
+		/// NOTE: Do not use this method for storing values, it is left in purely for compatibility. All package configuration
+		/// should be done using the target settings in the editor or in the [/Script/HoloLensPlatformEditor.HoloLensTargetSettings] INI section.
+		/// Values are stored in the following format.
+		/// Sample INI setting: Package.Capabilities.mx:Capability[0].Name=kinectAudio
+		/// Corresponding XML output:
+		///		&lt;Package&gt;
+		///		&lt;Capabilities&gt;
+		///		&lt;mx:Capability Name="kinectAudio"&gt;
+		/// Symbol Key:
+		/// $Section:Key$		Look up Key in Section in INI files (Game first then fall back to Engine).
+		/// 					Replace symbol with INI setting value.
+		/// %RelativeExePath%	Replace value with path to exe from package root.
+		/// %Insert:Path%		Insert contents of file at Path based off of project root path. Will indent all lines in the file
+		/// 					by the current value of Indent. File should contain valid XML (this is not verified).
+		/// </summary>
+		/// <param name="LookupString"> INI key to locate</param>
+		/// <param name="Index">Optional index of setting for [n] type LookupStrings</param>
+		/// <param name="Indent">Optional current indent count (used when writing %Insert:Path% values</param>
+		/// <returns>INI value for key LookupString post interpretation</returns>
+		private string GetInterprettedSettingValue(string LookupString, int Index = 0, int Indent = 0)
+		{
+			char[] VariableMarkers = { '$', '%' };
+			string BaseSetting;
+			string InterprettedSetting = "";
+
+			// Manifest settings are only (validly) located in Engine INI files
+			if (!EngineIni!.GetString("AppxManifest", LookupString, out BaseSetting))
+			{
+				return "";
+			}
+
+			// Parse results for any operators
+			int NextSetting = BaseSetting.IndexOfAny(VariableMarkers);
+			while (NextSetting >= 0)
+			{
+				// This will parse multiple operator types within a single setting, but not nested operators
+				if (NextSetting > 0)
+				{
+					// Copy any leading text (non-operator) to our output
+					InterprettedSetting += BaseSetting.Substring(0, NextSetting);
+				}
+				int LenOfSetting = BaseSetting.Substring(NextSetting + 1).IndexOfAny(VariableMarkers);
+				if (LenOfSetting < 0)
+				{
+					Log.TraceError("Could not parse setting {0}. Unmatched variable symbol '{1}'", LookupString, BaseSetting[NextSetting]);
+					return InterprettedSetting + BaseSetting;
+				}
+				if (BaseSetting[NextSetting] != BaseSetting[NextSetting + LenOfSetting + 1])
+				{
+					// Probable nested operators
+					Log.TraceError("Could not parse setting {0}. Mismatched variable symbols '{1}' and '{2}'", LookupString, BaseSetting[NextSetting], BaseSetting[NextSetting + LenOfSetting + 1]);
+					return InterprettedSetting + BaseSetting;
+				}
+
+				// Complete contents of operator
+				string VariableName = BaseSetting.Substring(NextSetting + 1, LenOfSetting);
+
+				switch (BaseSetting[NextSetting])
+				{
+					case '$':
+						// Look up $Section:Key$ in Game INIs
+						string IniSection = VariableName.Substring(0, VariableName.IndexOf(':'));
+						string IniSetting = VariableName.Substring(VariableName.IndexOf(':') + 1);
+						string IniValue;
+						GameIni!.GetString(IniSection, IniSetting, out IniValue);
+						// If not found in Game INIs, search for the same Key in Engine INIs
+						if (IniValue.Length == 0)
+						{
+							EngineIni.GetString(IniSection, IniSetting, out IniValue);
+						}
+						// Replace operator with value recovered
+						if (IniValue.Length == 0)
+						{
+							// @todo: Is there any better way to handle not finding the value? If we leave a value blank it will
+							// likely produce invalid XML and be difficult to trace. At least this hardcoded string should lead
+							// users back here.
+							InterprettedSetting += "InvalidIniValue";
+						}
+						else
+						{
+							InterprettedSetting += IniValue;
+						}
+						break;
+					case '%':
+						if (VariableName.StartsWith("Insert:"))
+						{
+							// Attempt to open path provided based off of the current project path
+							string InsertSource = Path.Combine(ProjectPath!, VariableName.Substring(VariableName.IndexOf(':') + 1));
+							if (!File.Exists(InsertSource))
+							{
+								Log.TraceWarning("Invalid path for insertion: {0}", InsertSource);
+								// @todo: Can't think of a way to insert valid XML in this case, so it's just going to be left out.
+								// It would be better at least to insert something that would lead back here (as is done with
+								// "InvalidIniValue" above.
+								break;
+							}
+							string[]? InsertContents = null;
+							try
+							{
+								InsertContents = File.ReadAllLines(InsertSource);
+							}
+							catch (Exception)
+							{
+								Log.TraceWarning("Error while trying to read data for insert from {0}.", InsertSource);
+								// @todo: Can't think of a way to insert valid XML in this case, so it's just going to be left out.
+								// It would be better at least to insert something that would lead back here (as is done with
+								// "InvalidIniValue" above.
+								break;
+							}
+							// Insert file contents one line at a time so that we can add indentation as needed.
+							foreach (string InsertLine in InsertContents)
+							{
+								InterprettedSetting += InsertLine + "\n";
+							}
+						}
+						else if (VariableName.StartsWith("ResourceString:") || VariableName.StartsWith("ResourceBinary:"))
+						{
+							string SectionKeyPair = VariableName.Substring(VariableName.IndexOf(':') + 1);
+							// Look up $Section:Key$ in Game INIs
+							string SettingSection = SectionKeyPair.Substring(0, SectionKeyPair.IndexOf(':'));
+							string SettingKey = SectionKeyPair.Substring(SectionKeyPair.IndexOf(':') + 1);
+							String? SettingValue = null;
+							GameIni!.GetString(SettingSection, SettingKey, out SettingValue);
+							// If not found in Game INIs, search for the same Key in Engine INIs
+							if (SettingValue == null || SettingValue.Length == 0)
+							{
+								EngineIni.GetString(SettingSection, SettingKey, out SettingValue);
+							}
+							// Replace operator with value recovered
+							if (SettingValue != null && SettingValue.Length > 0)
+							{
+								InterprettedSetting += SettingValue;
+							}
+						}
+						else if (VariableName.StartsWith("Array:"))
+						{
+							string SectionKeyPair = VariableName.Substring(VariableName.IndexOf(':') + 1);
+							// Look up $Section:Key$ in Game INIs
+							string ArraySection = SectionKeyPair.Substring(0, SectionKeyPair.IndexOf(':'));
+							string ArrayKey = SectionKeyPair.Substring(SectionKeyPair.IndexOf(':') + 1);
+							List<string>? ArraySettingValue = null;
+							GameIni!.GetArray(ArraySection, ArrayKey, out ArraySettingValue);
+							// If not found in Game INIs, search for the same Key in Engine INIs
+							if (ArraySettingValue == null || ArraySettingValue.Count == 0)
+							{
+								EngineIni.GetArray(ArraySection, ArrayKey, out ArraySettingValue);
+							}
+							// Replace operator with value recovered
+							if (ArraySettingValue == null || ArraySettingValue.Count == 0)
+							{
+								// @todo: Is there any better way to handle not finding the value? If we leave a value blank it will
+								// likely produce invalid XML and be difficult to trace. At least this hardcoded string should lead
+								// users back here.
+								InterprettedSetting += "InvalidIniValue";
+							}
+							else
+							{
+								List<string> ArraySettingValueDeduplicated = ArraySettingValue.Distinct().ToList();
+								if (ArraySettingValueDeduplicated.Count > Index)
+								{
+									InterprettedSetting += ArraySettingValueDeduplicated[Index];
+								}
+							}
+						}
+						else if (VariableName.StartsWith("AlphaNumericDot:"))
+						{
+							string SectionKeyPair = VariableName.Substring(VariableName.IndexOf(':') + 1);
+							// Look up $Section:Key$ in Game INIs
+							string SettingSection = SectionKeyPair.Substring(0, SectionKeyPair.IndexOf(':'));
+							string SettingKey = SectionKeyPair.Substring(SectionKeyPair.IndexOf(':') + 1);
+							string? SettingValue = null;
+							GameIni!.GetString(SettingSection, SettingKey, out SettingValue);
+							// If not found in Game INIs, search for the same Key in Engine INIs
+							if (SettingValue == null || SettingValue.Length == 0)
+							{
+								EngineIni.GetString(SettingSection, SettingKey, out SettingValue);
+							}
+							// Replace operator with value recovered
+							if (SettingValue == null || SettingValue.Length == 0)
+							{
+								// @todo: Is there any better way to handle not finding the value? If we leave a value blank it will
+								// likely produce invalid XML and be difficult to trace. At least this hardcoded string should lead
+								// users back here.
+								InterprettedSetting += "InvalidIniValue";
+							}
+							else
+							{
+								foreach (char Character in SettingValue.ToCharArray())
+								{
+									if ((Character >= 'A' && Character <= 'Z') ||
+									   (Character >= 'a' && Character <= 'z') ||
+									   (Character >= '0' && Character <= '9') ||
+									   (Character == '.'))
+									{
+										InterprettedSetting += Character;
+									}
+								}
+							}
+						}
+						else if (VariableName.StartsWith("DefaultValue:"))
+						{
+							int SectionIndex = VariableName.IndexOf(':') + 1;
+							int KeyIndex = VariableName.IndexOf(':', SectionIndex) + 1;
+							int ValueTypeIndex = VariableName.IndexOf(':', KeyIndex) + 1;
+							int DefaultValueIndex = VariableName.IndexOf(':', ValueTypeIndex) + 1;
+							string DefaultValue = VariableName.Substring(DefaultValueIndex);
+							string ValueType = VariableName.Substring(ValueTypeIndex, DefaultValueIndex - ValueTypeIndex - 1);
+							string SettingSection = VariableName.Substring(SectionIndex, KeyIndex - SectionIndex - 1);
+							string SettingKey = VariableName.Substring(KeyIndex, ValueTypeIndex - KeyIndex - 1);
+							// Look up $Section:Key$ in Game INIs
+							string? SettingValue = null;
+							if (ValueType.Equals("Int32", StringComparison.InvariantCultureIgnoreCase))
+							{
+								int Int32SettingValue;
+								GameIni!.GetInt32(SettingSection, SettingKey, out Int32SettingValue);
+								SettingValue = Int32SettingValue.ToString();
+							}
+							else if (ValueType.Equals("GUID", StringComparison.InvariantCultureIgnoreCase))
+							{
+								Guid GuidSettingValue;
+								GameIni!.TryGetValue(SettingSection, SettingKey, out GuidSettingValue);
+								SettingValue = GuidSettingValue.ToString("N");
+							}
+							else
+							{
+								GameIni!.GetString(SettingSection, SettingKey, out SettingValue);
+							}
+							// If not found in Game INIs, search for the same Key in Engine INIs
+							if (SettingValue == null || SettingValue.Length == 0)
+							{
+								if (ValueType.Equals("Int32", StringComparison.InvariantCultureIgnoreCase))
+								{
+									int Int32SettingValue;
+									EngineIni.GetInt32(SettingSection, SettingKey, out Int32SettingValue);
+									SettingValue = Int32SettingValue.ToString();
+								}
+								else if (ValueType.Equals("GUID", StringComparison.InvariantCultureIgnoreCase))
+								{
+									Guid GuidSettingValue;
+									EngineIni.TryGetValue(SettingSection, SettingKey, out GuidSettingValue);
+									SettingValue = GuidSettingValue.ToString("N");
+								}
+								else
+								{
+									EngineIni.GetString(SettingSection, SettingKey, out SettingValue);
+								}
+							}
+							// Replace operator with value recovered
+							if (SettingValue == null || SettingValue.Length == 0)
+							{
+								InterprettedSetting += DefaultValue;
+							}
+							else
+							{
+								InterprettedSetting += SettingValue;
+							}
+						}
+						else
+						{
+							Log.TraceWarning("Unable to parse AppxManifest variable value for {0}.", VariableName);
+							// @todo: Is there any better way to handle not finding the value? If we leave a value blank it will
+							// likely produce invalid XML and be difficult to trace. At least this hardcoded string should lead
+							// users back here.
+							InterprettedSetting += "InvalidVariableValue";
+						}
+						break;
+				}
+
+				// Find next operator pair (if any)
+				BaseSetting = BaseSetting.Substring(NextSetting + LenOfSetting + 2);
+				NextSetting = BaseSetting.IndexOfAny(VariableMarkers);
+			}
+
+			// Insert any tail (non-operator) text from the original setting to our output
+			InterprettedSetting += BaseSetting;
+
+			return InterprettedSetting;
+		}
+
+		/// <summary>
+		/// Checks that path is a directory and tries to create it if it doesn't exist.
+		/// </summary>
+		/// <returns>true if the directory is present and usable</returns>
+		private bool CreateCheckDirectory(string TargetDirectory)
+		{
+			if (!Directory.Exists(TargetDirectory))
+			{
+				try
+				{
+					Directory.CreateDirectory(TargetDirectory);
+				}
+				catch (Exception)
+				{
+					Log.TraceError("Could not create directory {0}.", TargetDirectory);
+					return false;
+				}
+				if (!Directory.Exists(TargetDirectory))
+				{
+					Log.TraceError("Path {0} does not exist or is not a directory.", TargetDirectory);
+					return false;
+				}
+			}
+			return true;
+		}
+
+		/// <summary>
+		/// Checks if an intermediate file has any modifications from the current target. Replaces the target file
+		/// if there are changes.
+		/// </summary>
+		private void CompareAndReplaceModifiedTarget(string IntermediatePath, string TargetPath)
+		{
+			if (!File.Exists(IntermediatePath))
+			{
+				Log.TraceError("Tried to copy non-existant intermediate file {0}.", IntermediatePath);
+				return;
+			}
+
+			CreateCheckDirectory(Path.GetDirectoryName(TargetPath)!);
+
+			// Check for differences in file contents
+			if (File.Exists(TargetPath))
+			{
+				byte[] OriginalContents = File.ReadAllBytes(TargetPath);
+				byte[] NewContents = File.ReadAllBytes(IntermediatePath);
+				if (!OriginalContents.SequenceEqual(NewContents))
+				{
+					try
+					{
+						FileAttributes attributes = File.GetAttributes(TargetPath);
+						if ((attributes & FileAttributes.ReadOnly) == FileAttributes.ReadOnly)
+						{
+							attributes &= ~FileAttributes.ReadOnly;
+							File.SetAttributes(TargetPath, attributes);
+						}
+						File.Delete(TargetPath);
+					}
+					catch (Exception)
+					{
+						Log.TraceError("Could not replace file {0}.", TargetPath);
+						return;
+					}
+				}
+			}
+
+			// If the file is present it is unmodified and should not be overwritten
+			if (!File.Exists(TargetPath))
+			{
+				try
+				{
+					File.Copy(IntermediatePath, TargetPath);
+				}
+				catch (Exception)
+				{
+					Log.TraceError("Unable to copy file {0}.", TargetPath);
+					return;
+				}
+				UpdatedFilePaths!.Add(TargetPath);
+			}
+		}
+
+		/// <summary>
+		/// Copies all cultures of a source resource to the intermediate directory.
+		/// <returns>true on success, false if the operation fails (i.e. the default source file doesn't exist)</returns>
+		/// </summary>
+		private bool CopyAndReplaceBinaryIntermediate(string ResourceFileName, bool AllowEngineFallback = true, Action<string, string>? CopyOp = null)
+		{
+			string TargetPath = Path.Combine(IntermediatePath!, BuildResourceSubPath);
+
+			CopyOp = CopyOp ?? File.Copy;
+
+			// look in project normal Build location
+			string SourcePath = Path.Combine(ProjectPath!, "Build", "HoloLens", BuildResourceSubPath);
+			bool bFileExists = File.Exists(Path.Combine(SourcePath, ResourceFileName));
+
+			// look in Platform Extensions next
+			if (!bFileExists)
+			{
+				SourcePath = Path.Combine(ProjectPath!, "Platforms", "HoloLens", "Build", BuildResourceSubPath);
+				bFileExists = File.Exists(Path.Combine(SourcePath, ResourceFileName));
+			}
+
+			// look in Engine, if allowed
+			if (!bFileExists && AllowEngineFallback)
+			{
+				SourcePath = Path.Combine(Unreal.EngineDirectory.FullName, "Build", "HoloLens", EngineResourceSubPath);
+				bFileExists = File.Exists(Path.Combine(SourcePath, ResourceFileName));
+
+				// look in Platform extensions too
+				if (!bFileExists)
+				{
+					SourcePath = Path.Combine(Unreal.EngineDirectory.FullName, "Platforms", "HoloLens", "Build", EngineResourceSubPath);
+					bFileExists = File.Exists(Path.Combine(SourcePath, ResourceFileName));
+				}
+			}
+
+			// At least the default culture entry for any resource binary must always exist
+			if (!bFileExists)
+			{
+				return false;
+			}
+
+			// If the target resource folder doesn't exist yet, create it
+			if (!CreateCheckDirectory(TargetPath))
+			{
+				return false;
+			}
+
+			// Find all copies of the resource file in the source directory (could be up to one for each culture and the default).
+			IEnumerable<string> SourceResourceInstances = Directory.EnumerateFiles(SourcePath, ResourceFileName, SearchOption.AllDirectories);
+
+			// Copy new resource files
+			foreach (string SourceResourceFile in SourceResourceInstances)
+			{
+				//@todo only copy files for cultures we are staging
+				string TargetResourcePath = Path.Combine(TargetPath, SourceResourceFile.Substring(SourcePath.Length + 1));
+				if (!CreateCheckDirectory(Path.GetDirectoryName(TargetResourcePath)!))
+				{
+					Log.TraceError("Unable to create intermediate directory {0}.", Path.GetDirectoryName(TargetResourcePath));
+					continue;
+				}
+				if (!File.Exists(TargetResourcePath))
+				{
+					try
+					{
+						CopyOp(SourceResourceFile, TargetResourcePath);
+					}
+					catch (Exception)
+					{
+						Log.TraceError("Unable to copy file {0} to {1}.", SourceResourceFile, TargetResourcePath);
+						return false;
+					}
+				}
+			}
+
+			// Now find specially named qualified versions of the resource (e.g. logo.scale-200.png) and give them the same treatment
+			string QualifiedResourceFileName = ResourceFileName.Replace(".png", ".*.png");
+			QualifiedResourceFileName = QualifiedResourceFileName.Replace(".glb", ".*.glb");
+			IEnumerable<string> SourceResourceQualifiedInstances = Directory.EnumerateFiles(SourcePath, QualifiedResourceFileName, SearchOption.AllDirectories);
+
+			// Copy new resource files
+			foreach (string SourceResourceFile in SourceResourceQualifiedInstances)
+			{
+				//@todo only copy files for cultures we are staging
+				string TargetResourcePath = Path.Combine(TargetPath, SourceResourceFile.Substring(SourcePath.Length + 1));
+				if (!CreateCheckDirectory(Path.GetDirectoryName(TargetResourcePath)!))
+				{
+					Log.TraceError("Unable to create intermediate directory {0}.", Path.GetDirectoryName(TargetResourcePath));
+					continue;
+				}
+				if (!File.Exists(TargetResourcePath))
+				{
+					try
+					{
+						CopyOp(SourceResourceFile, TargetResourcePath);
+					}
+					catch (Exception)
+					{
+						Log.TraceError("Unable to copy file {0} to {1}.", SourceResourceFile, TargetResourcePath);
+						throw;
+					}
+				}
+			}
+
+
+			return true;
+		}
+
+		/// <summary>
+		/// Copies modified intermediate resource binaries to the output path and deletes any stale resources in the
+		/// output path that do not exist in the intermediate directory.
+		/// <returns>A list of updated files.</returns>
+		/// </summary>
+		private void CopyResourcesToTargetDir()
+		{
+			string TargetPath = Path.Combine(OutputPath!, StoreResourceSubPath!);
+			string SourcePath = Path.Combine(IntermediatePath!, BuildResourceSubPath);
+
+			// If the target resource folder doesn't exist yet, create it
+			if (!CreateCheckDirectory(TargetPath))
+			{
+				return;
+			}
+
+			// Find all copies of the resource file in both target and source directories (could be up to one for each culture and the default, but must have at least the default).
+			IEnumerable<string> TargetResourceInstances = Directory.EnumerateFiles(TargetPath, "*.*", SearchOption.AllDirectories);
+			IEnumerable<string> SourceResourceInstances = Directory.EnumerateFiles(SourcePath, "*.*", SearchOption.AllDirectories);
+
+			// Remove any target files that aren't part of the source file list
+			foreach (string TargetResourceFile in TargetResourceInstances)
+			{
+				// Ignore string tables (the only non-binary resources that will be present
+				if (!TargetResourceFile.Contains(".resw"))
+				{
+					//@todo always delete for cultures we aren't staging
+					bool bRelativeSourceFileFound = false;
+					foreach (string SourceResourceFile in SourceResourceInstances)
+					{
+						string SourceRelativeFile = SourceResourceFile.Substring(SourcePath.Length + 1);
+						string TargetRelativeFile = TargetResourceFile.Substring(TargetPath.Length + 1);
+						if (SourceRelativeFile.Equals(TargetRelativeFile))
+						{
+							bRelativeSourceFileFound = true;
+							break;
+						}
+					}
+					if (!bRelativeSourceFileFound)
+					{
+						try
+						{
+							File.Delete(TargetResourceFile);
+						}
+						catch (Exception)
+						{
+							Log.TraceError("Could not remove stale resource file {0}.", TargetResourceFile);
+						}
+					}
+				}
+			}
+
+			// Copy new resource files only if they differ from the destination
+			foreach (string SourceResourceFile in SourceResourceInstances)
+			{
+				//@todo only copy files for cultures we are staging
+				string TargetResourcePath = Path.Combine(TargetPath, SourceResourceFile.Substring(SourcePath.Length + 1));
+				CompareAndReplaceModifiedTarget(SourceResourceFile, TargetResourcePath);
+				ManifestRelatedFilePaths!.Add(TargetResourcePath);
+			}
+		}
+
+		/// <summary>
+		/// Deletes a directory and everything it contains.
+		/// </summary>
+		/// <param name="InDirectoryToDelete">Directory to delete</param>
+		private void RecursivelyForceDeleteDirectory(string InDirectoryToDelete)
+		{
+			if (Directory.Exists(InDirectoryToDelete))
+			{
+				try
+				{
+					List<string> SubDirectories = new List<string>(Directory.GetDirectories(InDirectoryToDelete, "*.*", SearchOption.AllDirectories));
+					foreach (string DirectoryToRemove in SubDirectories)
+					{
+						RecursivelyForceDeleteDirectory(DirectoryToRemove);
+					}
+					List<string> FilesInDirectory = new List<string>(Directory.GetFiles(InDirectoryToDelete));
+					foreach (string FileToRemove in FilesInDirectory)
+					{
+						try
+						{
+							FileAttributes Attributes = File.GetAttributes(FileToRemove);
+							if ((Attributes & FileAttributes.ReadOnly) == FileAttributes.ReadOnly)
+							{
+								Attributes &= ~FileAttributes.ReadOnly;
+								File.SetAttributes(FileToRemove, Attributes);
+							}
+							File.Delete(FileToRemove);
+						}
+						catch (Exception)
+						{
+							Log.TraceWarning("Could not remove file {0} to remove directory {1}.", FileToRemove, InDirectoryToDelete);
+						}
+					}
+					Directory.Delete(InDirectoryToDelete, true);
+				}
+				catch (Exception)
+				{
+					Log.TraceWarning("Could not remove directory {0}.", InDirectoryToDelete);
+				}
+			}
+
+		}
+
+        // Add a new key with the given value to the input ini file.  The new value will take effect on the next build since the ini has already been ingested.
+        bool WriteProjectIniString(string iniFilePath, string SectionName, string keyName, string newString)
+        {
+            if (!File.Exists(iniFilePath))
+            {
+                return false;
+            }
+
+            string[] lines = File.ReadAllLines(iniFilePath);
+            List<string> updatedLines = new List<string>();
+            for (int i = 0; i < lines.Length; i++)
+            {
+                string line = lines[i].Trim();
+				if (line.Equals("[" + SectionName.Trim() + "]"))
+                {
+                    updatedLines = lines.ToList();
+					updatedLines.Insert(i + 1, string.Concat(keyName.Trim(), "=", newString.Trim()));
+
+                    break;
+                }
+            }
+
+            if (updatedLines.Count > 0)
+            {
+                File.WriteAllLines(iniFilePath, updatedLines.ToArray());
+                return true;
+            }
+
+            return false;
+        }
+
+        // Update any ini keyName entries in the project directory with newString.  The new value will take effect on the next build since the ini has already been ingested.
+        bool UpdateProjectIniString(FileReference? InProjectFile, UnrealTargetPlatform TargetPlatform, string SectionName, string keyName, string newString)
+        {
+            DirectoryReference? IniDirRef = DirectoryReference.FromFile(InProjectFile);
+            bool iniUpdated = false;
+
+            List<string> projectIniFilePaths = new List<string>();
+            //List<ConfigFile> Files = new List<ConfigFile>();
+            foreach (FileReference IniFileName in ConfigHierarchy.EnumerateConfigFileLocations(ConfigHierarchyType.Game, IniDirRef, TargetPlatform))
+            {
+                // If ini file does not exist, or is not in the project directory, move on to the next one.
+                if (!File.Exists(IniFileName.FullName)
+					|| !IniFileName.FullName.StartsWith(InProjectFile!.Directory.FullName))
+                {
+                    continue;
+                }
+
+                Log.TraceLog("Looking at file: " + IniFileName.FullName);
+                projectIniFilePaths.Add(IniFileName.FullName);
+
+                bool isInDesiredSection = false;
+                bool rewriteCurrentFile = false;
+
+                string[] lines = File.ReadAllLines(IniFileName.FullName);
+                for (int i = 0; i < lines.Length; i++)
+                {
+                    string line = lines[i].Trim();
+                    if (line.StartsWith("[") && line.EndsWith("]"))
+                    {
+						isInDesiredSection = line.Equals("[" + SectionName.Trim() + "]");
+                    }
+                    if (!isInDesiredSection)
+                    {
+                        continue;
+                    }
+
+					if (line.StartsWith(keyName.Trim() + "="))
+                    {
+                        Log.TraceLog("Found string match.");
+                        iniUpdated = true;
+                        rewriteCurrentFile = true;
+
+						lines[i] = string.Concat(keyName.Trim(), "=", newString.Trim());
+                    }
+                }
+
+                if (rewriteCurrentFile)
+                {
+                    File.WriteAllLines(IniFileName.FullName, lines);
+                }
+            }
+
+            if (!iniUpdated)
+            {
+                // The project ini did not have a project version entry.  This will happen on the first build which defaults to an entry in the engine ini.
+                foreach (string iniPath in projectIniFilePaths)
+                {
+                    if (WriteProjectIniString(iniPath, SectionName, keyName, newString))
+                    {
+                        iniUpdated = true;
+                    }
+                }
+            }
+
+            return iniUpdated;
+        }
+
+        /// <summary>
+        /// Kicks off manifest generation. Will always attempt to fully calculate a new manifest but will not update the output
+        /// file unless there are changes (to avoid unnecessary copies when deploying).
+        /// </summary>
+        /// <param name="TargetPlatform">The platform we're generating a manifest for.</param>
+        /// <param name="TargetArchitecture">The architecture we're generating a manifest for.</param>
+        /// <param name="InOutputPath">Path to write manifest files to.</param>
+        /// <param name="InIntermediatePath">Path to store temporary intermediate data (e.g. XML resource file).</param>
+        /// <param name="InProjectFile">Path to the uproject file</param>
+        /// <param name="InProjectDirectory">Directory containing the uproject file or the base engine path if no project file is specified (for content only builds).</param>
+        /// <param name="InTargetConfigs">Configurations to build manifest data for. Each configuration will generate it's own application entry.</param>
+        /// <param name="InExecutables">The launch executable for each configuration. Must match the length and order of InTargetConfigs.</param>
+        /// <param name="InWinMDReferences">The WinMD references that should be added as activatable types</param>
+        /// <returns>A list of all updated target files</returns>
+		[SupportedOSPlatform("windows")]
+        public List<string>? CreateManifest(UnrealTargetPlatform TargetPlatform, WindowsArchitecture TargetArchitecture, string InOutputPath, string InIntermediatePath, FileReference? InProjectFile, string InProjectDirectory, List<UnrealTargetConfiguration> InTargetConfigs, List<string> InExecutables, IEnumerable<WinMDRegistrationInfo>? InWinMDReferences)
+		{
+			// Check parameter values are valid
+			if (InTargetConfigs.Count != InExecutables.Count)
+			{
+				Log.TraceError("The number of target configurations ({0}) and executables ({1}) passed to manifest generation do not match.", InTargetConfigs.Count, InExecutables.Count);
+				return null;
+			}
+			if (File.Exists(InOutputPath))
+			{
+				Log.TraceWarning("InOutputPath {0} is a file. Should be a directory. Continuing using parent directory.", InOutputPath);
+				InOutputPath = Path.GetDirectoryName(InOutputPath)!;
+			}
+			if (File.Exists(InIntermediatePath))
+			{
+				Log.TraceWarning("InIntermediatePath {0} is a file. Should be a directory. Continuing using parent directory.", InIntermediatePath);
+				InIntermediatePath = Path.GetDirectoryName(InIntermediatePath)!;
+			}
+			if (!CreateCheckDirectory(InOutputPath))
+			{
+				return null;
+			}
+			if (!CreateCheckDirectory(InIntermediatePath))
+			{
+				return null;
+			}
+
+			OutputPath = InOutputPath;
+			IntermediatePath = InIntermediatePath;
+
+			UpdatedFilePaths = new List<string>();
+			ManifestRelatedFilePaths = new List<string>();
+
+			WinMDReferences = new List<WinMDRegistrationInfo>();
+			if(InWinMDReferences != null)
+            {
+				foreach(var r in InWinMDReferences)
+                {
+					if(!WinMDReferences.Any((a) => a.PackageRelativeDllPath.ToLower() == r.PackageRelativeDllPath.ToLower()))
+                    {
+						WinMDReferences.Add(r);
+					}
+                }
+            }
+			//Platform = TargetPlatform;
+			Architecture = TargetArchitecture;
+			TargetSettings = "/Script/HoloLensPlatformEditor.HoloLensTargetSettings";
+			BuildResourceSubPath = "Resources";
+			StoreResourceSubPath = WindowsExports.GetArchitectureSubpath(Architecture) + "\\" + BuildResourceSubPath;
+
+			// Clean out the resources intermediate path so that we know there are no stale binary files.
+			string IntermediateResourceDirectory = Path.Combine(IntermediatePath, BuildResourceSubPath);
+			RecursivelyForceDeleteDirectory(IntermediateResourceDirectory);
+			if (!Directory.Exists(IntermediateResourceDirectory))
+			{
+				try
+				{
+					Directory.CreateDirectory(IntermediateResourceDirectory);
+				}
+				catch (Exception)
+				{
+					Log.TraceError("Could not create directory {0}.", IntermediateResourceDirectory);
+					return null;
+				}
+			}
+
+			// Load up INI settings. We'll use engine settings to retrieve the manifest configuration, but these may reference
+			// values in either game or engine settings, so we'll keep both.
+			// Use the project directory here since this accounts for 'RemoteIniDir' when InProjectFile is null
+			if (InProjectFile != null)
+			{
+				IsDlc = InProjectFile.GetExtension() == ".uplugin";
+
+				if (IsDlc)
+				{
+					DirectoryReference IniDirRef = DirectoryReference.FromFile(InProjectFile).ParentDirectory!.ParentDirectory!;
+					GameIni = ConfigCache.ReadHierarchy(ConfigHierarchyType.Game, IniDirRef, TargetPlatform);
+					EngineIni = ConfigCache.ReadHierarchy(ConfigHierarchyType.Engine, IniDirRef, TargetPlatform);
+					IsGameSpecificExe = new DirectoryReference(InOutputPath).IsUnderDirectory(IniDirRef);
+
+					List<string>? DlcStoreMapping = new List<string>();
+					if (EngineIni.GetArray("/Script/HoloLensPlatformEditor.HoloLensTargetSettings", "DLCStoreMapping", out DlcStoreMapping))
+					{
+						foreach (string DlcEntry in DlcStoreMapping)
+						{
+							Dictionary<string, string> PossibleParsedDlcInfo = new Dictionary<string, string>();
+							InterpretINIStruct(DlcEntry, out PossibleParsedDlcInfo);
+							string? DlcName = null;
+							PossibleParsedDlcInfo.TryGetValue("PluginName", out DlcName);
+							if (DlcName == InProjectFile.GetFileNameWithoutExtension())
+							{
+								ParsedDlcInfo = PossibleParsedDlcInfo;
+								break;
+							}
+						}
+					}
+
+					if (ParsedDlcInfo == null)
+					{
+						Log.TraceWarning("Could not map {0} to a Store identity.  Using a temporary identity to enable local deployment.  For Store upload configure identity in the HoloLens Project Settings.", InProjectFile);
+						ParsedDlcInfo = new Dictionary<string, string>();
+						ParsedDlcInfo["PluginName"] = InProjectFile.GetFileNameWithoutExtension();
+						ParsedDlcInfo["PackageIdentityName"] = ParsedDlcInfo["PluginName"];
+						ParsedDlcInfo["PackageIdentityVersion"] = "1.0.0.0";
+					}
+				}
+				else
+				{
+					DirectoryReference IniDirRef = DirectoryReference.FromFile(InProjectFile);
+					GameIni = ConfigCache.ReadHierarchy(ConfigHierarchyType.Game, IniDirRef, TargetPlatform);
+					EngineIni = ConfigCache.ReadHierarchy(ConfigHierarchyType.Engine, IniDirRef, TargetPlatform);
+					IsGameSpecificExe = new DirectoryReference(InOutputPath).IsUnderDirectory(IniDirRef);
+				}
+			}
+			else if (!string.IsNullOrEmpty(UnrealBuildTool.GetRemoteIniPath()))
+			{
+				DirectoryReference IniDirRef = new DirectoryReference(UnrealBuildTool.GetRemoteIniPath()!);
+				GameIni = ConfigCache.ReadHierarchy(ConfigHierarchyType.Game, IniDirRef, TargetPlatform);
+				EngineIni = ConfigCache.ReadHierarchy(ConfigHierarchyType.Engine, IniDirRef, TargetPlatform);
+				IsGameSpecificExe = false;
+			}
+			else
+			{
+				GameIni = ConfigCache.ReadHierarchy(ConfigHierarchyType.Game, null, TargetPlatform);
+				EngineIni = ConfigCache.ReadHierarchy(ConfigHierarchyType.Engine, null, TargetPlatform);
+				IsGameSpecificExe = false;
+			}
+
+			ProjectPath = InProjectDirectory;
+			RelativeProjectRootForStage = IsGameSpecificExe ? InProjectFile!.GetFileNameWithoutAnyExtensions() : "Engine";
+
+			// Load and verify/clean culture list
+			List<string>? CulturesToStageWithDuplicates = null;
+			GameIni.GetArray("/Script/UnrealEd.ProjectPackagingSettings", "CulturesToStage", out CulturesToStageWithDuplicates);
+			if (CulturesToStageWithDuplicates == null || CulturesToStageWithDuplicates.Count < 1)
+			{
+				Log.TraceError("At least one culture must be selected to stage.");
+				return null;
+			}
+
+			CulturesToStage = CulturesToStageWithDuplicates.Distinct().ToList();
+
+			// Construct a single resource writer for the default (no-culture) values
+			string NeutralResourceIntermediatePath = Path.Combine(IntermediateResourceDirectory, "resources.resw");
+			NeutralResourceWriter = new UEResXWriter(NeutralResourceIntermediatePath);
+
+			PerCultureResourceWriters = new List<UEResXWriter>();
+			for (int i = 0; i < CulturesToStage.Count; ++i)
+			{
+				string Culture = CulturesToStage[i];
+				string IntermediateStringResourcePath = Path.Combine(IntermediateResourceDirectory, Culture);
+				string IntermediateStringResourceFile = Path.Combine(IntermediateStringResourcePath, "resources.resw");
+				if (!CreateCheckDirectory(IntermediateStringResourcePath))
+				{
+					Log.TraceWarning("Failed to create {0}.  Culture {1} resources not staged.", IntermediateStringResourcePath, Culture);
+					CulturesToStage.RemoveAt(i);
+					--i;
+					continue;
+				}
+				PerCultureResourceWriters.Add(new UEResXWriter(IntermediateStringResourceFile));
+			}
+
+			if (CulturesToStage.Count == 0)
+			{
+				Log.TraceError("Failed to create intermediate files for any culture.  Manifest could not be generated.");
+				return null;
+			}
+
+            // Check autoincrement flag
+            bool bAutoIncrementVersion;
+            if (EngineIni.GetBool(TargetSettings, "bAutoIncrementVersion", out bAutoIncrementVersion) && bAutoIncrementVersion)
+            {
+                string currentVersion;
+                if (GameIni.GetString("/Script/EngineSettings.GeneralProjectSettings", "ProjectVersion", out currentVersion))
+                {
+                    Log.TraceLog("Automatically incrementing version. Starting version is " + currentVersion);
+
+                    string[] versionEntries = currentVersion.Split(new char[] { '.' });
+                    if (versionEntries.Length == 4)
+                    {
+                        uint versionEntryToIncrement = Convert.ToUInt32(versionEntries[2]);
+                        versionEntryToIncrement++;
+
+                        string newVersion = string.Concat(versionEntries[0], ".", versionEntries[1], ".", versionEntryToIncrement, ".", versionEntries[3]);
+                        Log.TraceLog("Writing new version string: " + newVersion);
+
+                        // Update GameIni with newVersion.
+                        if (!UpdateProjectIniString(InProjectFile, TargetPlatform, "/Script/EngineSettings.GeneralProjectSettings", "ProjectVersion", newVersion))
+                        {
+                            Log.TraceWarning("Auto incrementing the project version was unsuccessful.");
+                        }
+                    }
+                    else
+                    {
+                        Log.TraceWarning("Auto increment was desired, but the number of decimals in the version string was unexpected.");
+                    }
+                }
+                else
+                {
+                    Log.TraceWarning("Auto increment was desired, but the existing version could not be identified.");
+                }
+            }
+
+            // Create the appxmanifest document
+            AppxManifestXmlDocument = new XmlDocument();
+			XmlDeclaration Declaration = AppxManifestXmlDocument.CreateXmlDeclaration("1.0", Encoding.UTF8.BodyName, null);
+			AppxManifestXmlDocument.AppendChild(Declaration);
+			
+			// Begin document content construction. Resources entries will be setup as required to support the manifest generation.
+			XmlNode Package = GetPackage(TargetPlatform, InTargetConfigs, InExecutables);
+			AppxManifestXmlDocument.AppendChild(Package);
+
+			// Export appxmanifest.xml to the intermediate directory then compare the contents to any existing target manifest
+			// and replace if there are differences.
+			string ManifestName = String.Format("AppxManifest_{0}.xml", WindowsExports.GetArchitectureSubpath(Architecture));
+			string ManifestIntermediatePath = Path.Combine(IntermediatePath, ManifestName);
+			string ManifestTargetPath = Path.Combine(OutputPath, ManifestName);
+			AppxManifestXmlDocument.Save(ManifestIntermediatePath);
+
+			// Check we produced a reasonable manifest document
+			ValidateAppxManifest(ManifestIntermediatePath);
+
+			CompareAndReplaceModifiedTarget(ManifestIntermediatePath, ManifestTargetPath);
+
+			// Clean out any resource directories that we aren't staging
+			string TargetResourcePath = Path.Combine(OutputPath, StoreResourceSubPath);
+			if (Directory.Exists(TargetResourcePath))
+			{
+				List<string> TargetResourceDirectories = new List<string>(Directory.GetDirectories(TargetResourcePath, "*.*", SearchOption.AllDirectories));
+				foreach (string ResourceDirectory in TargetResourceDirectories)
+				{
+					if (!CulturesToStage.Contains(Path.GetFileName(ResourceDirectory)))
+					{
+						RecursivelyForceDeleteDirectory(ResourceDirectory);
+					}
+				}
+			}
+
+			// MS staging code requires an AppxManifest to convert
+			//@todo remove if possible
+			// DLC packages do not contain an exe
+			if (InExecutables.Count > 0)
+			{
+				string ManifestBinaryPath = Path.Combine(Path.GetDirectoryName(InExecutables[0])!, ManifestName);
+				CompareAndReplaceModifiedTarget(ManifestIntermediatePath, ManifestBinaryPath);
+			}
+
+			// Export the resource tables starting with the neutral culture
+			string NeutralResourceTargetPath = Path.Combine(OutputPath, StoreResourceSubPath, "resources.resw");
+			NeutralResourceWriter.Close();
+			CompareAndReplaceModifiedTarget(NeutralResourceIntermediatePath, NeutralResourceTargetPath);
+
+			for (int CultureIndex = 0; CultureIndex < CulturesToStage.Count; CultureIndex++)
+			{
+				string Culture = CulturesToStage[CultureIndex];
+				string IntermediateStringResourceFile = Path.Combine(IntermediateResourceDirectory, Culture, "resources.resw");
+				string TargetStringResourceFile = Path.Combine(OutputPath, StoreResourceSubPath, Culture, "resources.resw");
+				PerCultureResourceWriters[CultureIndex].Close();
+				CompareAndReplaceModifiedTarget(IntermediateStringResourceFile, TargetStringResourceFile);
+			}
+
+			// Copy all the binary resources into the target directory.
+			CopyResourcesToTargetDir();
+
+			// The resource database is dependent on everything else calculated here (manifest, resource string tables, binary resources).
+			// So if any file has been updated we'll need to run the config.
+			IEnumerable<string> TargetPriFiles = Directory.EnumerateFiles(OutputPath, "*.pri");
+			if (UpdatedFilePaths.Count > 0 || TargetPriFiles.Count() == 0)
+			{
+				// Create resource index configuration
+				string PriExecutable = HoloLensToolChain.GetWindowsSdkToolPath("makepri.exe")!.FullName;
+
+				// We're not currently splitting pri files along the culture dimension, so all supported languages should be defaults
+				string AllDefaultCultures = CulturesToStage.Aggregate((c1, c2) => (c1 + "_" + c2));
+
+				string ResourceConfigFile = Path.Combine(IntermediatePath, "priconfig.xml");
+				string MakePriArgs = "createconfig /cf \"" + ResourceConfigFile + "\" /dq " + AllDefaultCultures + " /o /pv 10.0.0";
+				System.Diagnostics.ProcessStartInfo StartInfo = new System.Diagnostics.ProcessStartInfo(PriExecutable, MakePriArgs);
+				StartInfo.UseShellExecute = false;
+				StartInfo.RedirectStandardOutput = true;
+				StartInfo.CreateNoWindow = true;
+				int ExitCode = Utils.RunLocalProcessAndLogOutput(StartInfo);
+				if (ExitCode < 0)
+				{
+					throw new BuildException("Failed to generate config file for Package Resource Index.  See log for details.");
+				}
+
+				// Modify configuration to restrict indexing to the Resources directory (saves time and space)
+				XmlDocument PriConfig = new XmlDocument();
+				PriConfig.Load(ResourceConfigFile);
+
+				// Remove the 'Packaging' node so that we have no autoResourcePackage entries
+				XmlNodeList? PackagingNodes = PriConfig.SelectNodes("/resources/packaging");
+				if (PackagingNodes != null)
+				{
+					foreach (XmlNode? Node in PackagingNodes)
+					{
+						Node?.ParentNode?.RemoveChild(Node);
+					}
+				}
+
+				// The approach to limiting the indexer causes files to have dodgy uris in the generated pri e.g.
+				// ms-resource://PackageIdentityName/Files/Logo.png instead of ms-resource://PackageIdentityName/Files/Resources/Logo.png
+				// This appears to affect Windows's ability to locate a valid image in some scenarios such as a
+				// desktop shortcut.  So on HoloLens we start from the root and add exclusions.
+				XmlNodeList? ConfigNodes = PriConfig.SelectNodes("/resources/index/indexer-config");
+				if (ConfigNodes != null)
+				{
+					foreach (XmlNode? ConfigNode in ConfigNodes)
+					{
+						if (ConfigNode == null)
+						{
+							continue;
+						}
+
+						if (ConfigNode.Attributes?["type"]?.Value == "folder")
+						{
+							IEnumerable<string> AllSubItems = Directory.EnumerateFileSystemEntries(OutputPath);
+							foreach (string FileSystemEntry in AllSubItems)
+							{
+								if (Path.GetFileName(FileSystemEntry) != "Resources")
+								{
+									XmlElement ExcludeElement = PriConfig.CreateElement("exclude");
+									if (File.Exists(FileSystemEntry))
+									{
+										ExcludeElement.SetAttribute("type", "path");
+									}
+									else
+									{
+										ExcludeElement.SetAttribute("type", "tree");
+									}
+									ExcludeElement.SetAttribute("value", Path.GetFileName(FileSystemEntry));
+									ExcludeElement.SetAttribute("doNotTraverse", "true");
+									ExcludeElement.SetAttribute("doNotIndex", "true");
+									ConfigNode.AppendChild(ExcludeElement);
+								}
+							}
+						}
+					}
+				}
+
+				XmlNode ResNode = PriConfig.SelectSingleNode("/resources")!;
+				{
+					XmlAttribute Attr = PriConfig.CreateAttribute("isDeploymentMergeable");
+					Attr.Value = "true";
+					ResNode.Attributes!.Append(Attr);
+				}
+
+				PriConfig.Save(ResourceConfigFile);
+
+				// Remove previous pri files so we can enumerate which ones are new since the resource generator could produce a file for each staged language.
+				IEnumerable<string> OldPriFiles = Directory.EnumerateFiles(IntermediatePath, "*.pri");
+				foreach (string OldPri in OldPriFiles)
+				{
+					try
+					{
+						File.Delete(OldPri);
+					}
+					catch (Exception)
+					{
+						Log.TraceError("Could not delete file {0}.", OldPri);
+					}
+				}
+
+				// Generate the resource index
+				string ResourceLogFile = Path.Combine(IntermediatePath, "ResIndexLog_" + WindowsExports.GetArchitectureSubpath(Architecture) + ".xml");
+				string ResourceIndexFile = Path.Combine(IntermediatePath, "resources_" + WindowsExports.GetArchitectureSubpath(Architecture) + ".pri");
+				MakePriArgs = "new /pr \"" + IntermediateResourceDirectory + "\" /cf \"" + ResourceConfigFile + "\" /mn \"" + ManifestTargetPath + "\" /il \"" + ResourceLogFile + "\" /of \"" + ResourceIndexFile + "\" /o";
+				StartInfo = new System.Diagnostics.ProcessStartInfo(PriExecutable, MakePriArgs);
+				StartInfo.UseShellExecute = false;
+				StartInfo.RedirectStandardOutput = true;
+				StartInfo.CreateNoWindow = true;
+				StartInfo.StandardErrorEncoding = System.Text.Encoding.Unicode;
+				StartInfo.StandardOutputEncoding = System.Text.Encoding.Unicode;
+				ExitCode = Utils.RunLocalProcessAndLogOutput(StartInfo);
+				if (ExitCode < 0)
+				{
+					throw new BuildException("Failed to generate Package Resource Index file.  See log for details.");
+				}
+
+				// Stage all the modified pri files to the output directory
+				string FinalResourceIndexFile = Path.Combine(OutputPath, Path.GetFileName(ResourceIndexFile));
+				CompareAndReplaceModifiedTarget(ResourceIndexFile, FinalResourceIndexFile);
+			}
+
+			return ManifestRelatedFilePaths;
+		}
+
+		/// <summary>
+		/// Kicks off asset manifest generation. 
+		/// </summary>
+		/// <param name="TargetPlatform">The platform we're generating a manifest for.</param>
+		/// <param name="InOutputPath">Path to write manifest files to.</param>
+		/// <param name="InIntermediatePath">Path to store temporary intermediate data (e.g. XML resource file).</param>
+		/// <param name="InProjectFile">Path to the uproject file</param>
+		/// <param name="InProjectDirectory">Directory containing the uproject file or the base engine path if no project file is specified (for content only builds).</param>
+		/// <returns>A list of all updated target files</returns>
+		[SupportedOSPlatform("windows")]
+		public List<string>? CreateAssetsManifest(UnrealTargetPlatform TargetPlatform, string InOutputPath, string InIntermediatePath, FileReference InProjectFile, string InProjectDirectory)
+		{
+			if (File.Exists(InOutputPath))
+			{
+				Log.TraceWarning("InOutputPath {0} is a file. Should be a directory. Continuing using parent directory.", InOutputPath);
+				InOutputPath = Path.GetDirectoryName(InOutputPath)!;
+			}
+			if (File.Exists(InIntermediatePath))
+			{
+				Log.TraceWarning("InIntermediatePath {0} is a file. Should be a directory. Continuing using parent directory.", InIntermediatePath);
+				InIntermediatePath = Path.GetDirectoryName(InIntermediatePath)!;
+			}
+			if (!CreateCheckDirectory(InOutputPath))
+			{
+				return null;
+			}
+			if (!CreateCheckDirectory(InIntermediatePath))
+			{
+				return null;
+			}
+
+			OutputPath = InOutputPath;
+			IntermediatePath = InIntermediatePath;
+
+			UpdatedFilePaths = new List<string>();
+			ManifestRelatedFilePaths = new List<string>();
+
+			WinMDReferences = new List<WinMDRegistrationInfo>();
+			//Platform = TargetPlatform;
+			TargetSettings = "/Script/HoloLensPlatformEditor.HoloLensTargetSettings";
+			BuildResourceSubPath = "AssetsResources";
+			StoreResourceSubPath =  BuildResourceSubPath;
+
+			string IntermediateResourceDirectory = Path.Combine(IntermediatePath, BuildResourceSubPath);
+			RecursivelyForceDeleteDirectory(IntermediateResourceDirectory);
+			if (!Directory.Exists(IntermediateResourceDirectory))
+			{
+				try
+				{
+					Directory.CreateDirectory(IntermediateResourceDirectory);
+				}
+				catch (Exception)
+				{
+					Log.TraceError("Could not create directory {0}.", IntermediateResourceDirectory);
+					return null;
+				}
+			}
+
+			// Load up INI settings. We'll use engine settings to retrieve the manifest configuration, but these may reference
+			// values in either game or engine settings, so we'll keep both.
+			// Use the project directory here since this accounts for 'RemoteIniDir' when InProjectFile is null
+			if (InProjectFile != null)
+			{
+				IsDlc = InProjectFile.GetExtension() == ".uplugin";
+
+				if (IsDlc)
+				{
+					DirectoryReference IniDirRef = DirectoryReference.FromFile(InProjectFile).ParentDirectory!.ParentDirectory!;
+					GameIni = ConfigCache.ReadHierarchy(ConfigHierarchyType.Game, IniDirRef, TargetPlatform);
+					EngineIni = ConfigCache.ReadHierarchy(ConfigHierarchyType.Engine, IniDirRef, TargetPlatform);
+					IsGameSpecificExe = new DirectoryReference(InOutputPath).IsUnderDirectory(IniDirRef);
+
+					List<string>? DlcStoreMapping = new List<string>();
+					if (EngineIni.GetArray("/Script/HoloLensPlatformEditor.HoloLensTargetSettings", "DLCStoreMapping", out DlcStoreMapping))
+					{
+						foreach (string DlcEntry in DlcStoreMapping)
+						{
+							Dictionary<string, string> PossibleParsedDlcInfo = new Dictionary<string, string>();
+							InterpretINIStruct(DlcEntry, out PossibleParsedDlcInfo);
+							string? DlcName = null;
+							PossibleParsedDlcInfo.TryGetValue("PluginName", out DlcName);
+							if (DlcName == InProjectFile.GetFileNameWithoutExtension())
+							{
+								ParsedDlcInfo = PossibleParsedDlcInfo;
+								break;
+							}
+						}
+					}
+
+					if (ParsedDlcInfo == null)
+					{
+						Log.TraceWarning("Could not map {0} to a Store identity.  Using a temporary identity to enable local deployment.  For Store upload configure identity in the HoloLens Project Settings.", InProjectFile);
+						ParsedDlcInfo = new Dictionary<string, string>();
+						ParsedDlcInfo["PluginName"] = InProjectFile.GetFileNameWithoutExtension();
+						ParsedDlcInfo["PackageIdentityName"] = ParsedDlcInfo["PluginName"];
+						ParsedDlcInfo["PackageIdentityVersion"] = "1.0.0.0";
+					}
+				}
+				else
+				{
+					DirectoryReference IniDirRef = DirectoryReference.FromFile(InProjectFile);
+					GameIni = ConfigCache.ReadHierarchy(ConfigHierarchyType.Game, IniDirRef, TargetPlatform);
+					EngineIni = ConfigCache.ReadHierarchy(ConfigHierarchyType.Engine, IniDirRef, TargetPlatform);
+					IsGameSpecificExe = new DirectoryReference(InOutputPath).IsUnderDirectory(IniDirRef);
+				}
+			}
+			else if (!string.IsNullOrEmpty(UnrealBuildTool.GetRemoteIniPath()))
+			{
+				DirectoryReference IniDirRef = new DirectoryReference(UnrealBuildTool.GetRemoteIniPath()!);
+				GameIni = ConfigCache.ReadHierarchy(ConfigHierarchyType.Game, IniDirRef, TargetPlatform);
+				EngineIni = ConfigCache.ReadHierarchy(ConfigHierarchyType.Engine, IniDirRef, TargetPlatform);
+				IsGameSpecificExe = false;
+			}
+			else
+			{
+				GameIni = ConfigCache.ReadHierarchy(ConfigHierarchyType.Game, null, TargetPlatform);
+				EngineIni = ConfigCache.ReadHierarchy(ConfigHierarchyType.Engine, null, TargetPlatform);
+				IsGameSpecificExe = false;
+			}
+
+			ProjectPath = InProjectDirectory;
+			RelativeProjectRootForStage = IsGameSpecificExe ? InProjectFile!.GetFileNameWithoutAnyExtensions() : "Engine";
+
+			// Load and verify/clean culture list
+			List<string>? CulturesToStageWithDuplicates = null;
+			GameIni.GetArray("/Script/UnrealEd.ProjectPackagingSettings", "CulturesToStage", out CulturesToStageWithDuplicates);
+			if (CulturesToStageWithDuplicates == null || CulturesToStageWithDuplicates.Count < 1)
+			{
+				Log.TraceError("At least one culture must be selected to stage.");
+				return null;
+			}
+
+			CulturesToStage = CulturesToStageWithDuplicates.Distinct().ToList();
+
+			// Construct a single resource writer for the default (no-culture) values
+			string NeutralResourceIntermediatePath = Path.Combine(IntermediateResourceDirectory, "resources.resw");
+			NeutralResourceWriter = new UEResXWriter(NeutralResourceIntermediatePath);
+
+			PerCultureResourceWriters = new List<UEResXWriter>();
+			for (int i = 0; i < CulturesToStage.Count; ++i)
+			{
+				string Culture = CulturesToStage[i];
+				string IntermediateStringResourcePath = Path.Combine(IntermediateResourceDirectory, Culture);
+				string IntermediateStringResourceFile = Path.Combine(IntermediateStringResourcePath, "resources.resw");
+				if (!CreateCheckDirectory(IntermediateStringResourcePath))
+				{
+					Log.TraceWarning("Failed to create {0}.  Culture {1} resources not staged.", IntermediateStringResourcePath, Culture);
+					CulturesToStage.RemoveAt(i);
+					--i;
+					continue;
+				}
+				PerCultureResourceWriters.Add(new UEResXWriter(IntermediateStringResourceFile));
+			}
+
+			if (CulturesToStage.Count == 0)
+			{
+				Log.TraceError("Failed to create intermediate files for any culture.  Manifest could not be generated.");
+				return null;
+			}
+
+			AppxManifestXmlDocument = new XmlDocument();
+			var Declaration = AppxManifestXmlDocument.CreateXmlDeclaration("1.0", Encoding.UTF8.BodyName, null);
+			AppxManifestXmlDocument.AppendChild(Declaration);
+
+			// Begin document content construction. Resources entries will be setup as required to support the manifest generation.
+			XmlElement Package = AppxManifestXmlDocument.CreateElement("Package");
+			{
+				XmlAttribute ManifestNamespace = AppxManifestXmlDocument.CreateAttribute("xmlns");
+				ManifestNamespace.Value = "http://schemas.microsoft.com/appx/manifest/foundation/windows10";
+				Package.Attributes.Append(ManifestNamespace);
+				Version? WinBuild;
+
+				XmlElement Identity = AppxManifestXmlDocument.CreateElement("Identity");
+				{
+					XmlAttribute PackageName = CreateStringAttribute("Name", "PackageName", "Package.Identity.Name", "/Script/EngineSettings.GeneralProjectSettings", "ProjectName", "DefaultUnrealProject", ValidatePackageName);
+					Identity.Attributes.Append(PackageName);
+					XmlAttribute PublisherName = CreateStringAttribute("Publisher", "PublisherName", "Package.Identity.Publisher", "/Script/EngineSettings.GeneralProjectSettings", "CompanyDistinguishedName", "CN=NoPublisher");
+					Identity.Attributes.Append(PublisherName);
+					XmlAttribute VersionNumber = CreateStringAttribute("Version", "PackageVersion", "Package.Identity.Version", "/Script/EngineSettings.GeneralProjectSettings", "ProjectVersion", "1.0.0.0");
+					Identity.Attributes.Append(VersionNumber);
+				}
+				Package.AppendChild(Identity);
+
+				var Dependencies = AppxManifestXmlDocument.CreateElement("Dependencies");
+				{
+					XmlElement TargetDeviceFamily = AppxManifestXmlDocument.CreateElement("TargetDeviceFamily");
+					Dependencies.AppendChild(TargetDeviceFamily);
+
+					XmlAttribute NameAttribute = AppxManifestXmlDocument.CreateAttribute("Name");
+					NameAttribute.Value = CreateStringValue("TargetDeviceFamily", "Package.Dependencies.TargetDeviceFamily[0].Name", "TargetDeviceFamily", "Name", "Windows.Universal");
+					TargetDeviceFamily.Attributes.Append(NameAttribute);
+
+					XmlAttribute MinVersionAttribute = AppxManifestXmlDocument.CreateAttribute("MinVersion");
+					string versionString = CreateStringValue("MinimumPlatformVersion", "Package.Dependencies.TargetDeviceFamily[0].MinVersion", "MinimumPlatformVersion", "MinVersion", "10.0.17763.0");
+					MinVersionAttribute.Value = versionString;
+					TargetDeviceFamily.Attributes.Append(MinVersionAttribute);
+
+					Version.TryParse(versionString, out WinBuild);
+
+					XmlAttribute MaxVersionTestedAttribute = AppxManifestXmlDocument.CreateAttribute("MaxVersionTested");
+					MaxVersionTestedAttribute.Value = CreateStringValue("MaximumPlatformVersionTested", "Package.Dependencies.TargetDeviceFamily[0].MaxVersionTested", "MaximumPlatformVersionTested", "MaxVersionTested", "10.0.10586.0");
+					TargetDeviceFamily.Attributes.Append(MaxVersionTestedAttribute);
+				}
+				Package.AppendChild(Dependencies);
+
+
+				var Properties = AppxManifestXmlDocument.CreateElement("Properties");
+				{
+
+					XmlElement DisplayName = AppxManifestXmlDocument.CreateElement("DisplayName");
+					DisplayName.InnerText = "ms-resource:PackageDisplayName";
+					Properties.AppendChild(DisplayName);
+
+					XmlElement PublisherDisplayName = AppxManifestXmlDocument.CreateElement("PublisherDisplayName");
+					PublisherDisplayName.InnerText = "ms-resource:PublisherDisplayName";
+					Properties.AppendChild(PublisherDisplayName);
+
+					XmlElement PackageLogo = AppxManifestXmlDocument.CreateElement("Logo");
+                    if (CopyAndReplaceBinaryIntermediate("StoreLogo.png"))
+                    {
+						PackageLogo.InnerText = BuildResourceSubPath + "\\StoreLogo.png";
+						Properties.AppendChild(PackageLogo);
+					}
+					else
+					{
+						Log.TraceError("Unable to stage package logo.");
+					}
+
+					XmlElement ResourcePackageParam = AppxManifestXmlDocument.CreateElement("ResourcePackage");
+					ResourcePackageParam.InnerText = "true";
+					Properties.AppendChild(ResourcePackageParam);
+
+					if(WinBuild != null && WinBuild >= new Version("10.0.17134.0"))
+					{
+						XmlElement NotExecutionParam = AppxManifestXmlDocument.CreateElement("uap6", "AllowExecution","http://schemas.microsoft.com/appx/manifest/uap/windows10/6");
+						NotExecutionParam.InnerText = "false";
+						Properties.AppendChild(NotExecutionParam);
+					}
+				}
+				Package.AppendChild(Properties);
+			}
+			AppxManifestXmlDocument.AppendChild(Package);
+
+			// Export appxmanifest.xml to the intermediate directory then compare the contents to any existing target manifest
+			// and replace if there are differences.
+			string ManifestTargetPath = Path.Combine(InOutputPath, "AppxManifest_assets.xml");
+			AppxManifestXmlDocument.Save(ManifestTargetPath);
+
+			// Check we produced a reasonable manifest document
+			ValidateAppxManifest(ManifestTargetPath);
+
+			// Copy all the binary resources into the target directory.
+			CopyResourcesToTargetDir();
+
+			return ManifestRelatedFilePaths;
+		}
+
+		/// <summary>
+		/// Interpret a struct stored within a single INI entry into key-value pairs.
+		/// Note: Will parse all nested structs.
+		/// </summary>
+		private void InterpretINIStruct(string INIStruct, out Dictionary<string, string> StructDictionary)
+		{
+			StructDictionary = new Dictionary<string, string>();
+			char[] EntryEndChars = { '(', ')', ',' };
+			string[] StructEntries = INIStruct.Split(EntryEndChars, StringSplitOptions.RemoveEmptyEntries);
+			foreach (string CurrentEntry in StructEntries)
+			{
+				int AssignmentIndex = CurrentEntry.IndexOf('=');
+				if (AssignmentIndex > 0)
+				{
+					string EntryKey = CurrentEntry.Substring(0, AssignmentIndex);
+					string EntryValue = CurrentEntry.Substring(AssignmentIndex + 1);
+					char[] EntryTrimChars = { '"' };
+					EntryValue = EntryValue.Trim(EntryTrimChars);
+					StructDictionary.Add(EntryKey, EntryValue);
+				}
+			}
+		}
+		
+		/// <summary>
+		/// Calculate a manifest string value based on a system of fallback possibilities and return it.
+		/// Selection priorities:
+		/// 1. PlatformINIKey under section [/Script/HoloLensPlatformEditor.HoloLensTargetSettings]
+		/// 2. Deprecated [AppxManifest] setting where the key equals ManifestFullPath
+		/// 3. Engine INI value with section equal to GenericINISection and key equal to GenericINIKey
+		/// 4. Game INI value with section equal to GenericINISection and key equal to GenericINIKey
+		/// 5. The DefaultValue passed in
+		/// </summary>
+		private string CreateStringValue(string PlatformINIKey, string ManifestFullPath, string GenericINISection, string GenericINIKey, string DefaultValue, Func<string, string>? ValueValidationDelegate = null)
+		{
+			string ConfigScratchValue = "";
+			if (!EngineIni!.GetString(TargetSettings!, PlatformINIKey, out ConfigScratchValue) || ConfigScratchValue.Length <= 0)
+			{
+				if (ManifestFullPath != null)
+				{
+					ConfigScratchValue = GetInterprettedSettingValue(ManifestFullPath);
+				}
+				if (ConfigScratchValue == null || ConfigScratchValue.Length <= 0)
+				{
+					// If a config value wasn't specified, don't try to read from the configs and just go with the default value
+					if (GenericINISection != null && GenericINIKey != null)
+					{
+						bool EngineConfigReadSuccess = EngineIni.GetString(GenericINISection, GenericINIKey, out ConfigScratchValue);
+						// If the engine config read failed or the returned value was empty/null, keep searching, otherwise use the value we already retrieved in ConfigScratchValue
+						if (!EngineConfigReadSuccess || ConfigScratchValue == null || ConfigScratchValue.Length <= 0)
+						{
+							bool GameConfigReadSuccess = GameIni!.GetString(GenericINISection, GenericINIKey, out ConfigScratchValue);
+							// If the game config read failed or the returned value was empty/null, use the default value, otherwise use the value we already retrieved in ConfigScratchValue
+							if (!GameConfigReadSuccess || ConfigScratchValue == null || ConfigScratchValue.Length <= 0)
+							{
+								ConfigScratchValue = DefaultValue;
+							}
+						}
+					}
+					else
+					{
+						ConfigScratchValue = DefaultValue;
+					}
+				}
+			}
+			if (ValueValidationDelegate != null)
+			{
+				return ValueValidationDelegate(ConfigScratchValue);
+			}
+			else
+			{
+				return ConfigScratchValue;
+			}
+		}
+
+		/// <summary>
+		/// Calculate a manifest string value using CreateStringValue and compare that value to TrueValue, return the result.
+		/// </summary>
+		private bool CreateBoolValue(string PlatformINIKey, string ManifestFullPath, string GenericINISection, string GenericINIKey, string DefaultValue, string TrueValue = "True")
+		{
+			string ConfigScratchValue = CreateStringValue(PlatformINIKey, ManifestFullPath, GenericINISection, GenericINIKey, DefaultValue);
+			return ConfigScratchValue.Equals(TrueValue, StringComparison.InvariantCultureIgnoreCase);
+		}
+
+		/// <summary>
+		/// Calculate a manifest string value using CreateStringValue and return the result as an integer.
+		/// </summary>
+		private int CreateIntValue(string PlatformINIKey, string ManifestFullPath, string GenericINISection, string GenericINIKey, string DefaultValue)
+		{
+			string ConfigScratchValue = CreateStringValue(PlatformINIKey, ManifestFullPath, GenericINISection, GenericINIKey, DefaultValue);
+			return Int32.Parse(ConfigScratchValue);
+		}
+
+		/// <summary>
+		/// Calculate a manifest array value based on a system of fallback possibilities and return it.
+		/// Selection priorities:
+		/// 1. PlatformINIKey under section [/Script/HoloLensPlatformEditor.HoloLensTargetSettings]
+		/// 2. Deprecated [AppxManifest] setting where the key equals ManifestFullPath
+		/// 3. Engine INI value with section equal to GenericINISection and key equal to GenericINIKey
+		/// 4. Game INI value with section equal to GenericINISection and key equal to GenericINIKey
+		/// 5. The DefaultValue passed in
+		/// </summary>
+		private List<string> CreateArrayValue(string PlatformINIKey, string ManifestPath, string ManifestSubKey, string GenericINISection, string GenericINIKey, List<string> DefaultValue)
+		{
+			List<string>? ConfigScratchValue = null;
+			if (!EngineIni!.GetArray(TargetSettings!, PlatformINIKey, out ConfigScratchValue))
+			{
+				if (ManifestPath != null)
+				{
+					// Retrieve the deprecated [AppxManifest] value, but it's an array so we will have to pull from multiple entries
+					int ArrayIndex = 0;
+					while (true)
+					{
+						string FullManifestPath = ManifestPath + "[" + ArrayIndex + "]";
+						if (ManifestSubKey != null && ManifestSubKey.Length <= 0)
+						{
+							FullManifestPath += "." + ManifestSubKey;
+						}
+						string ArrayElement = GetInterprettedSettingValue(FullManifestPath);
+						if (ArrayElement == null || ArrayElement.Length <= 0)
+						{
+							break;
+						}
+						if (ConfigScratchValue == null)
+						{
+							ConfigScratchValue = new List<string>(1);
+						}
+						ConfigScratchValue.Add(ArrayElement);
+						ArrayIndex++;
+					}
+				}
+				if (ConfigScratchValue == null || ConfigScratchValue.Count <= 0)
+				{
+					if (GenericINISection == null || GenericINIKey == null || (!EngineIni.GetArray(GenericINISection, GenericINIKey, out ConfigScratchValue) && !GameIni!.GetArray(GenericINISection, GenericINIKey, out ConfigScratchValue)))
+					{
+						ConfigScratchValue = DefaultValue;
+					}
+				}
+			}
+			return ConfigScratchValue;
+		}
+
+		/// <summary>
+		/// Calculate a manifest string value using CreateStringValue and return it as an XmlElement.
+		/// </summary>
+		private XmlElement CreateStringElement(string ElementName, string PlatformINIKey, string ManifestFullPath, string GenericINISection, string GenericINIKey, string DefaultValue)
+		{
+			XmlElement TargetElement = AppxManifestXmlDocument!.CreateElement(ElementName);
+			string ConfigScratchValue = CreateStringValue(PlatformINIKey, ManifestFullPath, GenericINISection, GenericINIKey, DefaultValue);
+			TargetElement.InnerText = ConfigScratchValue;
+			return TargetElement;
+		}
+
+		/// <summary>
+		/// Calculate a manifest string value using CreateStringValue and return it as an XmlAttribute.
+		/// </summary>
+		private XmlAttribute CreateStringAttribute(string ElementName, string PlatformINIKey, string ManifestFullPath, string GenericINISection, string GenericINIKey, string DefaultValue, Func<string, string>? ValueValidationDelegate = null)
+		{
+			XmlAttribute TargetAttribute = AppxManifestXmlDocument!.CreateAttribute(ElementName);
+			string ConfigScratchValue = CreateStringValue(PlatformINIKey, ManifestFullPath, GenericINISection, GenericINIKey, DefaultValue, ValueValidationDelegate);
+			TargetAttribute.Value = ConfigScratchValue;
+			return TargetAttribute;
+		}
+
+		/// <summary>
+		/// Calculate a manifest hex color value based on a system of fallback possibilities and return an attribute containing it.
+		/// Selection priorities:
+		/// 1. PlatformINIKey under section [/Script/HoloLensPlatformEditor.HoloLensTargetSettings]
+		/// 2. Deprecated [AppxManifest] setting where the key equals ManifestFullPath
+		/// 3. The DefaultValue passed in
+		/// </summary>
+		private XmlAttribute CreateColorAttribute(string ElementName, string PlatformINIKey, string ManifestFullPath, string DefaultValue)
+		{
+			string ColorValue = "";
+			if (EngineIni!.GetString(TargetSettings!, PlatformINIKey, out ColorValue))
+			{
+				// Break the setting down by color
+				Dictionary<string, string> StructValues;
+				InterpretINIStruct(ColorValue, out StructValues);
+				int Red = 0;
+				string RedText = StructValues["R"];
+				Int32.TryParse(RedText, out Red);
+				int Green = 0;
+				string GreenText = StructValues["G"];
+				Int32.TryParse(GreenText, out Green);
+				int Blue = 0;
+				string BlueText = StructValues["B"];
+				Int32.TryParse(BlueText, out Blue);
+				ColorValue = "#" + Red.ToString("X2") + Green.ToString("X2") + Blue.ToString("X2");
+			}
+			else
+			{
+				ColorValue = GetInterprettedSettingValue(ManifestFullPath);
+				if (ColorValue == null || ColorValue.Length <= 0)
+				{
+					ColorValue = DefaultValue;
+				}
+			}
+
+			XmlAttribute ColorAttribute = AppxManifestXmlDocument!.CreateAttribute("BackgroundColor");
+			ColorAttribute.Value = ColorValue;
+			return ColorAttribute;
+		}
+
+		/// <summary>
+		/// Using the old style manifest settings, loop through all elements in an array looking for a subkey that matches the specified value.
+		/// </summary>
+		/// <returns>The index of the first matching array element. -1 if no match is found.</returns>
+		private int FindIndexOfArrayValue(string ManifestArrayKey, string ManifestTestElement, string ManifestTestValue)
+		{
+			int ArrayIndex = 0;
+			while (true)
+			{
+				string ManifestValue = GetInterprettedSettingValue(ManifestArrayKey + "[" + ArrayIndex + "]." + ManifestTestElement);
+				if (ManifestValue == null || ManifestValue.Length == 0)
+				{
+					return -1;
+				}
+				if (ManifestValue.Equals(ManifestTestValue, StringComparison.InvariantCultureIgnoreCase))
+				{
+					return ArrayIndex;
+				}
+				ArrayIndex++;
+			}
+		}
+
+		/// <summary>
+		/// Determine whether a resource is per-DLC or should be inherited from the main package
+		/// </summary>
+		private bool UseDlcResourcesForResourceEntry(string ResourceEntryName)
+		{
+			return IsDlc &&
+				(ResourceEntryName == "PackageDisplayName" ||
+				ResourceEntryName == "PackageDescription");
+		}
+
+		/// <summary>
+		/// Calculate the per culture manifest string value and add a resource table entry encompassing the values.
+		/// </summary>
+		private void AddResourceEntry(string ResourceEntryName, string ConfigKey, string ManifestFullPath, string GenericINISection, string GenericINIKey, string DefaultValue, string ValuePostfix = "")
+		{
+			// Enter the default (no-culture) value
+			string ConfigScratchValue = "";
+			string DefaultCultureScratchValue = "";
+			if (EngineIni!.GetString(TargetSettings!, "CultureStringResources", out DefaultCultureScratchValue))
+			{
+				Dictionary<string, string> DefaultCultureStringValues;
+				InterpretINIStruct(DefaultCultureScratchValue, out DefaultCultureStringValues);
+				ConfigScratchValue = DefaultCultureStringValues[ConfigKey];
+			}
+			if (ConfigScratchValue == null || ConfigScratchValue.Length <= 0)
+			{
+				if (ManifestFullPath != null)
+				{
+					ConfigScratchValue = GetInterprettedSettingValue(ManifestFullPath);
+				}
+				if (ConfigScratchValue == null || ConfigScratchValue.Length <= 0)
+				{
+					if (GenericINISection == null || GenericINIKey == null || (!EngineIni.GetString(GenericINISection, GenericINIKey, out ConfigScratchValue) && !GameIni!.GetString(GenericINISection, GenericINIKey, out ConfigScratchValue)))
+					{
+						ConfigScratchValue = DefaultValue;
+					}
+				}
+			}
+
+			bool IsDlcDefinedResource = UseDlcResourcesForResourceEntry(ResourceEntryName);
+
+			Dictionary<string, string> IniLocalizedValues = new Dictionary<string, string>();
+			List<string>? PerCultureValues;
+			if (EngineIni.GetArray(TargetSettings!, IsDlcDefinedResource ? "DlcPerCultureResources" : "PerCultureResources", out PerCultureValues))
+			{
+				foreach (string CultureCombinedValues in PerCultureValues)
+				{
+					Dictionary<string, string> SeparatedCultureValues;
+					InterpretINIStruct(CultureCombinedValues, out SeparatedCultureValues);
+					if (!IsDlcDefinedResource ||
+						string.Compare(SeparatedCultureValues["AppliesToDlcPlugin"], ParsedDlcInfo!["PluginName"], StringComparison.InvariantCultureIgnoreCase) == 0)
+					{
+						string CultureId = SeparatedCultureValues["CultureId"];
+						bool IsStagedCulture = string.IsNullOrEmpty(CultureId) || CulturesToStage!.Contains(CultureId);
+						if (IsStagedCulture)
+						{
+							if (SeparatedCultureValues[ConfigKey] != null && SeparatedCultureValues[ConfigKey].Length > 0)
+							{
+								IniLocalizedValues.Add(CultureId, SeparatedCultureValues[ConfigKey] + ValuePostfix);
+							}
+						}
+					}
+				}
+			}
+
+			string? NeutralValue = string.Empty;
+			if (!IniLocalizedValues.TryGetValue("", out NeutralValue))
+			{
+				NeutralValue = ConfigScratchValue + ValuePostfix;
+			}
+
+			bool IsEverLocalized = false;
+			for (int i = 0; i < CulturesToStage!.Count; ++i)
+			{
+				string? ValueToWrite = string.Empty;
+				if (IniLocalizedValues.TryGetValue(CulturesToStage[i], out ValueToWrite))
+				{
+					IsEverLocalized = true;
+					PerCultureResourceWriters![i].AddResource(ResourceEntryName, ValueToWrite);
+				}
+				else
+				{ 
+					Log.TraceVerbose("No localized value for {0} in culture {1}.  Neutral value ({2}) will be used", ResourceEntryName, CulturesToStage[i], NeutralValue);
+				}
+			}
+
+            // Values cannot be empty in the resource file, or the appx will fail WACK.
+            if (string.IsNullOrEmpty(NeutralValue.Trim()))
+            {
+                if (!string.IsNullOrEmpty(DefaultValue.Trim()))
+                {
+                    NeutralValue = DefaultValue.Trim();
+                }
+                else
+                {
+                    NeutralValue = "Missing Entry";
+                }
+            }
+
+            // Any culture with a culture-specific value will override the neutral value,
+            // even for unrelated cultures.  So propagate the neutral value to avoid this happening.
+            if (IsEverLocalized)
+			{
+				for (int i = 0; i < CulturesToStage.Count; ++i)
+				{
+					//string ValueToWrite = string.Empty;
+					if (!IniLocalizedValues.ContainsKey(CulturesToStage[i]))
+					{
+						PerCultureResourceWriters![i].AddResource(ResourceEntryName, NeutralValue);
+					}
+				}
+			}
+			else
+			{
+				// No culture has a specific value for this string.  Write the shared value to the neutral resource collection.
+				NeutralResourceWriter!.AddResource(ResourceEntryName, NeutralValue);
+			}
+		}
+
+		/// <summary>
+		/// Add a child XML node to a parent node if the child node is valid.
+		/// <param name="Parent">Parent node to add child to.</param>
+		/// <param name="Child">Child node to be evaluated and added.</param>
+		/// <param name="bNodeRequired">Display an error if the node is invalid and won't be added to the parent.</param>
+		/// <param name="bNodeMustNotBeEmpty">The child node must contain child nodes of it's own to be considered valid.</param>
+		/// </summary>
+		private void AddElementIfValid(XmlNode Parent, XmlNode? Child, bool bNodeRequired, bool bNodeMustNotBeEmpty = false)
+		{
+			if (Child != null)
+			{
+				if (!bNodeMustNotBeEmpty || Child.HasChildNodes)
+				{
+					Parent.AppendChild(Child);
+				}
+				else if (bNodeRequired)
+				{
+					Log.TraceError("Node {0} that requires a value is empty.", Child.Name);
+				}
+			}
+			else if (bNodeRequired)
+			{
+				Log.TraceError("Unable to create required manifest entry");
+			}
+		}
+
+		/// <summary>
+		/// Gather all information for the Package element of the manifest.
+		/// </summary>
+		private XmlNode GetPackage(UnrealTargetPlatform TargetPlatform, List<UnrealTargetConfiguration> TargetConfigs, List<string> Executables)
+		{
+			XmlElement Package = AppxManifestXmlDocument!.CreateElement("Package");
+
+			XmlAttribute ManifestNamespace = AppxManifestXmlDocument.CreateAttribute("xmlns");
+			ManifestNamespace.Value = "http://schemas.microsoft.com/appx/manifest/foundation/windows10";
+			Package.Attributes.Append(ManifestNamespace);
+
+			XmlAttribute UapManifestNamespace = AppxManifestXmlDocument.CreateAttribute("xmlns:uap");
+			UapManifestNamespace.Value = "http://schemas.microsoft.com/appx/manifest/uap/windows10";
+			Package.Attributes.Append(UapManifestNamespace);
+
+            XmlAttribute Uap2ManifestNamespace = AppxManifestXmlDocument.CreateAttribute("xmlns:uap2");
+            Uap2ManifestNamespace.Value = "http://schemas.microsoft.com/appx/manifest/uap/windows10/2";
+            Package.Attributes.Append(Uap2ManifestNamespace);
+
+			XmlAttribute Uap3ManifestNamespace = AppxManifestXmlDocument.CreateAttribute("xmlns:uap3");
+			Uap3ManifestNamespace.Value = "http://schemas.microsoft.com/appx/manifest/uap/windows10/3";
+			Package.Attributes.Append(Uap3ManifestNamespace);
+
+			XmlAttribute Uap4ManifestNamespace = AppxManifestXmlDocument.CreateAttribute("xmlns:uap4");
+			Uap4ManifestNamespace.Value = "http://schemas.microsoft.com/appx/manifest/uap/windows10/4";
+			Package.Attributes.Append(Uap4ManifestNamespace);
+
+			XmlAttribute Uap5ManifestNamespace = AppxManifestXmlDocument.CreateAttribute("xmlns:uap5");
+			Uap5ManifestNamespace.Value = "http://schemas.microsoft.com/appx/manifest/uap/windows10/5";
+			Package.Attributes.Append(Uap5ManifestNamespace);
+
+			XmlAttribute MpManifestNamespace = AppxManifestXmlDocument.CreateAttribute("xmlns:mp");
+			MpManifestNamespace.Value = "http://schemas.microsoft.com/appx/2014/phone/manifest";
+			Package.Attributes.Append(MpManifestNamespace);
+
+            XmlAttribute IgnorableNamespaces = AppxManifestXmlDocument.CreateAttribute("IgnorableNamespaces");
+            IgnorableNamespaces.Value = "mp uap uap2 uap3 uap4 uap5";
+            Package.Attributes.Append(IgnorableNamespaces);
+
+
+            XmlNode Identity = GetIdentity();
+			AddElementIfValid(Package, Identity, true);
+
+			XmlNode Properties = GetProperties();
+			AddElementIfValid(Package, Properties, true);
+
+			XmlNode Dependencies = GetDependencies(TargetConfigs);
+			AddElementIfValid(Package, Dependencies, true);
+
+			XmlNode Resources = GetResources();
+			AddElementIfValid(Package, Resources, true);
+
+			if (!IsDlc)
+			{
+				XmlNode Applications = GetApplications(TargetConfigs, Executables);
+				AddElementIfValid(Package, Applications, true);
+
+				XmlNode Capabilities = GetCapabilities(TargetPlatform);
+				AddElementIfValid(Package, Capabilities, true);
+
+				XmlNode? Extensions = GetPackageExtensions();
+				AddElementIfValid(Package, Extensions, false, true);
+			}
+			return Package;
+		}
+
+		/// <summary>
+		/// Validate a package name. Must contain only characters [-.A-Za-z0-9].
+		/// </summary>
+		private string ValidatePackageName(string InPackageName)
+		{
+			string ReturnVal = Regex.Replace(InPackageName, "[^-.A-Za-z0-9]", "");
+			if (ReturnVal == null || ReturnVal.Length <= 0)
+			{
+				Log.TraceError("Invalid package name {0}. Package names must only contain letters, numbers, dash, and period and must be at least one character long.", InPackageName);
+				Log.TraceError("Consider using the setting [/Script/HoloLensPlatformEditor.HoloLensTargetSettings]:PackageName to provide a HoloLens specific value.");
+				ReturnVal = String.Empty;
+			}
+			return ReturnVal;
+		}
+
+		/// <summary>
+		/// Gather all information for the Identity element of the manifest.
+		/// </summary>
+		private XmlNode GetIdentity()
+		{
+			XmlElement Identity = AppxManifestXmlDocument!.CreateElement("Identity");
+
+			if (!IsDlc)
+			{
+				XmlAttribute PackageName = CreateStringAttribute("Name", "PackageName", "Package.Identity.Name", "/Script/EngineSettings.GeneralProjectSettings", "ProjectName", "DefaultUnrealProject", ValidatePackageName);
+				Identity.Attributes.Append(PackageName);
+			}
+			else
+			{
+				XmlAttribute PackageName = AppxManifestXmlDocument.CreateAttribute("Name");
+				PackageName.Value = ParsedDlcInfo!["PackageIdentityName"];
+				ValidatePackageName(PackageName.Value);
+				Identity.Attributes.Append(PackageName);
+			}
+
+			XmlAttribute ProcessorArchitecture = AppxManifestXmlDocument.CreateAttribute("ProcessorArchitecture");
+			// @MIXEDREALITY_CHANGE : BEGIN TODO:
+			ProcessorArchitecture.Value = WindowsExports.GetArchitectureSubpath(Architecture); 
+			Identity.Attributes.Append(ProcessorArchitecture);
+
+			XmlAttribute PublisherName = CreateStringAttribute("Publisher", "PublisherName", "Package.Identity.Publisher", "/Script/EngineSettings.GeneralProjectSettings", "CompanyDistinguishedName", "CN=NoPublisher");
+			Identity.Attributes.Append(PublisherName);
+
+			if (!IsDlc)
+			{
+				XmlAttribute VersionNumber = CreateStringAttribute("Version", "PackageVersion", "Package.Identity.Version", "/Script/EngineSettings.GeneralProjectSettings", "ProjectVersion", "1.0.0.0");
+				Identity.Attributes.Append(VersionNumber);
+			}
+			else
+			{
+				XmlAttribute VersionNumber = AppxManifestXmlDocument.CreateAttribute("Version");
+				VersionNumber.Value = ParsedDlcInfo!["PackageIdentityVersion"];
+				Identity.Attributes.Append(VersionNumber);
+			}
+
+			return Identity;
+		}
+
+		/// <summary>
+		/// Gather all information for the Properties element of the manifest.
+		/// </summary>
+		private XmlNode GetProperties()
+		{
+			XmlElement Properties = AppxManifestXmlDocument!.CreateElement("Properties");
+
+			XmlElement DisplayName = AppxManifestXmlDocument.CreateElement("DisplayName");
+			DisplayName.InnerText = "ms-resource:PackageDisplayName";
+			Properties.AppendChild(DisplayName);
+			AddResourceEntry("PackageDisplayName", "PackageDisplayName", "Package.Properties.DisplayName", "/Script/EngineSettings.GeneralProjectSettings", "ProjectDisplayedTitle", "DefaultUnrealProject");
+
+			XmlElement PublisherDisplayName = AppxManifestXmlDocument.CreateElement("PublisherDisplayName");
+			PublisherDisplayName.InnerText = "ms-resource:PublisherDisplayName";
+			Properties.AppendChild(PublisherDisplayName);
+			AddResourceEntry("PublisherDisplayName", "PublisherDisplayName", "Package.Properties.PublisherDisplayName", "/Script/EngineSettings.GeneralProjectSettings", "CompanyName", "NoPublisher");
+
+			XmlElement PackageDescription = AppxManifestXmlDocument.CreateElement("Description");
+			PackageDescription.InnerText = "ms-resource:PackageDescription";
+			Properties.AppendChild(PackageDescription);
+			AddResourceEntry("PackageDescription", "PackageDescription", "Package.Properties.Description", "/Script/EngineSettings.GeneralProjectSettings", "Description", "No Description");
+
+			XmlElement PackageLogo = AppxManifestXmlDocument.CreateElement("Logo");
+			if (CopyAndReplaceBinaryIntermediate("StoreLogo.png"))
+			{
+				PackageLogo.InnerText = BuildResourceSubPath + "\\StoreLogo.png";
+				Properties.AppendChild(PackageLogo);
+			}
+			else
+			{
+				Log.TraceError("Unable to stage package logo.");
+			}
+
+			return Properties;
+		}
+
+		/// <summary>
+		/// Gather all information for the Dependencies element of the manifest.
+		/// </summary>
+		private XmlNode GetDependencies(List<UnrealTargetConfiguration> TargetConfigs)
+		{
+			XmlElement Dependencies = AppxManifestXmlDocument!.CreateElement("Dependencies");
+			
+			{
+				XmlElement TargetDeviceFamily = AppxManifestXmlDocument.CreateElement("TargetDeviceFamily");
+				Dependencies.AppendChild(TargetDeviceFamily);
+
+				XmlAttribute NameAttribute = AppxManifestXmlDocument.CreateAttribute("Name");
+				NameAttribute.Value = CreateStringValue("TargetDeviceFamily", "Package.Dependencies.TargetDeviceFamily[0].Name", "TargetDeviceFamily", "Name", "Windows.Universal");
+				TargetDeviceFamily.Attributes.Append(NameAttribute);
+
+				XmlAttribute MinVersionAttribute = AppxManifestXmlDocument.CreateAttribute("MinVersion");
+				MinVersionAttribute.Value = CreateStringValue("MinimumPlatformVersion", "Package.Dependencies.TargetDeviceFamily[0].MinVersion", "MinimumPlatformVersion", "MinVersion", "10.0.17763.0");
+				TargetDeviceFamily.Attributes.Append(MinVersionAttribute);
+
+				XmlAttribute MaxVersionTestedAttribute = AppxManifestXmlDocument.CreateAttribute("MaxVersionTested");
+				MaxVersionTestedAttribute.Value = CreateStringValue("MaximumPlatformVersionTested", "Package.Dependencies.TargetDeviceFamily[0].MaxVersionTested", "MaximumPlatformVersionTested", "MaxVersionTested", "10.0.18362.0");
+				TargetDeviceFamily.Attributes.Append(MaxVersionTestedAttribute);
+			}
+
+			{
+				XmlElement PackageDependency = AppxManifestXmlDocument.CreateElement("PackageDependency");
+				Dependencies.AppendChild(PackageDependency);
+
+				XmlAttribute NameAttribute = AppxManifestXmlDocument.CreateAttribute("Name");
+				NameAttribute.Value = "Microsoft.VCLibs.140.00";
+				PackageDependency.Attributes.Append(NameAttribute);
+
+				XmlAttribute PublisherAttribute = AppxManifestXmlDocument.CreateAttribute("Publisher");
+				PublisherAttribute.Value = "CN=Microsoft Corporation, O=Microsoft Corporation, L=Redmond, S=Washington, C=US";
+				PackageDependency.Attributes.Append(PublisherAttribute);
+
+				XmlAttribute MinVersionAttribute = AppxManifestXmlDocument.CreateAttribute("MinVersion");
+				MinVersionAttribute.Value = "14.0.0.0";
+				PackageDependency.Attributes.Append(MinVersionAttribute);
+			}
+
+			if (IsDlc)
+			{
+				// Add a dependency for the main package
+				XmlElement PackageDependency = AppxManifestXmlDocument.CreateElement("uap3:MainPackageDependency", "http://schemas.microsoft.com/appx/manifest/uap/windows10/3");
+				Dependencies.AppendChild(PackageDependency);
+
+				XmlAttribute NameAttribute = CreateStringAttribute("Name", "PackageName", "Package.Identity.Name", "/Script/EngineSettings.GeneralProjectSettings", "ProjectName", "DefaultUnrealProject", ValidatePackageName);
+				PackageDependency.Attributes.Append(NameAttribute);
+			}
+
+			return Dependencies;
+		}
+ 
+
+		/// <summary>
+		/// Gather all information for the Prerequisites element of the manifest.
+		/// </summary>
+		private XmlNode GetPrerequisites()
+		{
+			XmlElement Prerequisites = AppxManifestXmlDocument!.CreateElement("Prerequisites");
+			XmlElement OSMinVersion = CreateStringElement("OSMinVersion", "MinimumOSVersion", "Package.Prerequisites.OSMinVersion", TargetSettings!, "MinimumOSVersion", "6.2");
+			Prerequisites.AppendChild(OSMinVersion);
+
+			XmlElement OSMaxVersionTested = CreateStringElement("OSMaxVersionTested", "MaximumOSVersion", "Package.Prerequisites.OSMaxVersionTested", TargetSettings!, "MaximumOSVersion", "6.2");
+			Prerequisites.AppendChild(OSMaxVersionTested);
+
+			return Prerequisites;
+		}
+
+		/// <summary>
+		/// Gather all information for the Resources element of the manifest.
+		/// </summary>
+		private XmlNode GetResources()
+		{
+			XmlElement Resources = AppxManifestXmlDocument!.CreateElement("Resources");
+
+			// Check that we have a valid number of cultures.
+			if (CulturesToStage!.Count < 1 || CulturesToStage.Count >= MaxResourceEntries)
+			{
+				Log.TraceWarning("Incorrect number of cultures to stage. There must be between 1 and {0} cultures selected.", MaxResourceEntries);
+			}
+
+			// Create the culture list.
+			for (int ResourceIndex = 0; ResourceIndex < CulturesToStage.Count; ResourceIndex++)
+			{
+				XmlNode Resource = AppxManifestXmlDocument.CreateElement("Resource");
+
+				XmlAttribute LanguageAttribute = AppxManifestXmlDocument.CreateAttribute("Language");
+				LanguageAttribute.Value = CulturesToStage[ResourceIndex];
+				Resource.Attributes!.Append(LanguageAttribute);
+
+				Resources.AppendChild(Resource);
+			}
+
+			return Resources;
+		}
+
+		/// <summary>
+		/// Create an Application manifest entry for each target configuration
+		/// </summary>
+		private XmlNode GetApplications(List<UnrealTargetConfiguration> TargetConfigs, List<string> Executables)
+		{
+			XmlElement Applications = AppxManifestXmlDocument!.CreateElement("Applications");
+
+			if (TargetConfigs.Count < 1)
+			{
+				Log.TraceError("No configurations to deploy");
+				return Applications;
+			}
+			if (TargetConfigs.Count != Executables.Count)
+			{
+				Log.TraceError("The number of executables does not match the number of configurations.");
+				return Applications;
+			}
+
+			for (int ApplicationIndex = 0; ApplicationIndex < TargetConfigs.Count; ApplicationIndex++)
+			{
+				bool bIncludeConfigPostfix = TargetConfigs.Count > 1 && TargetConfigs[ApplicationIndex] != UnrealTargetConfiguration.Development;
+				XmlNode Application = GetApplication(ApplicationIndex, TargetConfigs[ApplicationIndex], Executables[ApplicationIndex], bIncludeConfigPostfix);
+				AddElementIfValid(Applications, Application, true, true);
+			}
+
+			return Applications;
+		}
+
+		/// <summary>
+		/// Validate the base name we use to construct the application id and entry point. Must match [A-Za-z][A-Za-z0-9]*.
+		/// </summary>
+		private string ValidateApplicationName(string InApplicationId)
+		{
+			string ReturnVal = Regex.Replace(InApplicationId, "[^A-Za-z0-9]", "");
+			if (ReturnVal != null)
+			{
+				// Remove any leading numbers (must start with a letter)
+				ReturnVal = Regex.Replace(ReturnVal, "^[0-9]*", "");
+			}
+			if (ReturnVal == null || ReturnVal.Length <= 0)
+			{
+				Log.TraceError("Invalid application ID {0}. Application IDs must only contain letters and numbers. And they must begin with a letter.", InApplicationId);
+				Log.TraceError("Consider using the setting [/Script/HoloLensPlatformEditor.HoloLensTargetSettings]:ValidateApplicationName to provide a HoloLens specific value.");
+				ReturnVal = String.Empty;
+			}
+			return ReturnVal;
+		}
+
+		/// <summary>
+		/// Create an Application manifest entry for a specific target configuration
+		/// </summary>
+		private XmlNode GetApplication(int ApplicationIndex, UnrealTargetConfiguration TargetConfig, string ExecutablePath, bool bIncludeConfigPostfix)
+		{
+			XmlElement Application = AppxManifestXmlDocument!.CreateElement("Application");
+
+			string PackageBaseName = CreateStringValue("ApplicationName", "Package.Applications.Application[" + ApplicationIndex + "].Id", "/Script/EngineSettings.GeneralProjectSettings", "ProjectName", "UnrealGame", ValidateApplicationName);
+
+			string ConfigPostfix = "";
+			if (bIncludeConfigPostfix)
+			{
+				ConfigPostfix = TargetConfig.ToString();
+			}
+
+			string MakeRelativeTo = IsGameSpecificExe ? Path.Combine(ProjectPath!, "..") : Unreal.EngineDirectory.FullName;
+			string RelativeExePath = IsGameSpecificExe ?
+				Utils.MakePathRelativeTo(ExecutablePath, MakeRelativeTo) :
+				Path.Combine(RelativeProjectRootForStage!, Utils.MakePathRelativeTo(ExecutablePath, MakeRelativeTo));
+
+			XmlAttribute Id = AppxManifestXmlDocument.CreateAttribute("Id");
+			Id.Value = "App" + PackageBaseName + ConfigPostfix;
+			Application.Attributes.Append(Id);
+
+			XmlAttribute Executable = AppxManifestXmlDocument.CreateAttribute("Executable");
+			Executable.Value = RelativeExePath;
+			Application.Attributes.Append(Executable);
+
+			XmlAttribute EntryPoint = AppxManifestXmlDocument.CreateAttribute("EntryPoint");
+			EntryPoint.Value = PackageBaseName + ".app";
+			Application.Attributes.Append(EntryPoint);
+
+			XmlNode VisualElements = GetVisualElements(/*Document, */ApplicationIndex, ConfigPostfix);
+			AddElementIfValid(Application, VisualElements, true, true);
+
+			return Application;
+		}
+
+		/// <summary>
+		/// Gather all information for the VisualElements element of the manifest.
+		/// </summary>
+		private XmlNode GetVisualElements(int ApplicationIndex, string ConfigPostfix)
+		{
+			XmlElement VisualElements = AppxManifestXmlDocument!.CreateElement("uap:VisualElements", "http://schemas.microsoft.com/appx/manifest/uap/windows10");
+
+			XmlAttribute DisplayName = AppxManifestXmlDocument.CreateAttribute("DisplayName");
+			if (ConfigPostfix != null && ConfigPostfix.Length > 0)
+			{
+				DisplayName.Value = "ms-resource:AppDisplayName" + ConfigPostfix;
+				AddResourceEntry("AppDisplayName" + ConfigPostfix, "ApplicationDisplayName", "Package.Applications.Application[" + ApplicationIndex + "].VisualElements.DisplayName", "/Script/EngineSettings.GeneralProjectSettings", "ProjectName", "UnrealGame", " - " + ConfigPostfix);
+			}
+			else
+			{
+				DisplayName.Value = "ms-resource:AppDisplayName";
+				AddResourceEntry("AppDisplayName" + ConfigPostfix, "ApplicationDisplayName", "Package.Applications.Application[" + ApplicationIndex + "].VisualElements.DisplayName", "/Script/EngineSettings.GeneralProjectSettings", "ProjectName", "UnrealGame");
+			}
+			VisualElements.Attributes.Append(DisplayName);
+
+			XmlAttribute Description = AppxManifestXmlDocument.CreateAttribute("Description");
+			Description.Value = "ms-resource:AppDescription";
+			VisualElements.Attributes.Append(Description);
+			AddResourceEntry("AppDescription", "ApplicationDescription", "Package.Applications.Application[" + ApplicationIndex + "].VisualElements.Description", "/Script/EngineSettings.GeneralProjectSettings", "Description", "");
+
+			XmlAttribute BackgroundColor = CreateColorAttribute("BackgroundColor", "TileBackgroundColor", "Package.Applications.Application[" + ApplicationIndex + "].VisualElements.BackgroundColor", "#000040");
+			VisualElements.Attributes.Append(BackgroundColor);
+
+			XmlAttribute Logo = AppxManifestXmlDocument.CreateAttribute("Square150x150Logo");
+			if (CopyAndReplaceBinaryIntermediate("Logo.png"))
+			{
+				Logo.Value = BuildResourceSubPath + "\\Logo.png";
+				VisualElements.Attributes.Append(Logo);
+			}
+			else
+			{
+				Log.TraceError("Unable to stage application logo.");
+			}
+
+			XmlAttribute SmallLogo = AppxManifestXmlDocument.CreateAttribute("Square44x44Logo");
+			if (CopyAndReplaceBinaryIntermediate("SmallLogo.png"))
+			{
+				SmallLogo.Value = BuildResourceSubPath + "\\SmallLogo.png";
+				VisualElements.Attributes.Append(SmallLogo);
+			}
+			else
+			{
+				Log.TraceError("Unable to stage application small logo.");
+			}
+
+			XmlNode SplashScreen = GetSplashScreen(ApplicationIndex);
+			VisualElements.AppendChild(SplashScreen);
+
+			XmlNode DefaultTile = GetDefaultTile(ApplicationIndex);
+			VisualElements.AppendChild(DefaultTile);
+
+			//@todo application support
+			// 			XmlNode ViewStates = GetViewStates(Document, ApplicationIndex);
+			// 			VisualElements.AppendChild(ViewStates);
+
+			return VisualElements;
+		}
+
+		/// <summary>
+		/// Gather all information for the DefaultTile element of the manifest.
+		/// </summary>
+		private XmlNode GetDefaultTile(int ApplicationIndex)
+		{
+			XmlElement DefaultTile = AppxManifestXmlDocument!.CreateElement("uap:DefaultTile", "http://schemas.microsoft.com/appx/manifest/uap/windows10");
+			bool bIsWideLogoUsed = false;
+
+			XmlAttribute WideLogo = AppxManifestXmlDocument.CreateAttribute("Wide310x150Logo");
+			if (CopyAndReplaceBinaryIntermediate("WideLogo.png"))
+			{
+				WideLogo.Value = BuildResourceSubPath + "\\WideLogo.png";
+				DefaultTile.Attributes.Append(WideLogo);
+				bIsWideLogoUsed = true;
+				if (CopyAndReplaceBinaryIntermediate("3DLogo.glb", false, (string from, string to) => 
+				{
+					//we need to process the logo with the gltf tool
+					string Args = string.Format("\"{0}\" -o \"{1}\" -platform all -replace-textures -min-version 1803" //-temp-directory \"{2}\"
+						, from
+						, to
+						//, Path.Combine(IntermediatePath, "tempGltf") //temp path doesn't work
+						);
+					string Exe = Path.Combine(Unreal.EngineDirectory.FullName, "Binaries", "Win64", "WindowsMRAssetConverter.exe");
+					System.Diagnostics.ProcessStartInfo StartInfo = new System.Diagnostics.ProcessStartInfo(Exe, Args);
+					StartInfo.UseShellExecute = false;
+					StartInfo.RedirectStandardOutput = true;
+					StartInfo.CreateNoWindow = true;
+					int ExitCode = Utils.RunLocalProcessAndPrintfOutput(StartInfo);
+					if (ExitCode < 0)
+					{
+						Log.TraceError("GLTF packaging failed. See log for details.");
+						throw new BuildException("GLTF packaging failed. See log for details.");
+					}
+					File.SetLastWriteTimeUtc(to, File.GetLastWriteTimeUtc(from));
+				}))
+				{
+					XmlElement MixedRealityModel = AppxManifestXmlDocument.CreateElement("uap5:MixedRealityModel", "http://schemas.microsoft.com/appx/manifest/uap/windows10/5");
+					XmlAttribute MixedRealityModelPath = AppxManifestXmlDocument.CreateAttribute("Path");
+					MixedRealityModelPath.Value = BuildResourceSubPath + "\\3DLogo.glb";
+					MixedRealityModel.Attributes.Append(MixedRealityModelPath);
+					DefaultTile.AppendChild(MixedRealityModel);
+				}
+			}
+			else
+			{
+				Log.TraceError("Unable to stage application wide logo.");
+			}
+
+			bool bUseNameForLogo;
+			if(EngineIni!.GetBool(TargetSettings!, "bUseNameForLogo", out bUseNameForLogo) && bUseNameForLogo)
+			{
+				XmlElement ShowNameOnTiles = AppxManifestXmlDocument.CreateElement("uap:ShowNameOnTiles", "http://schemas.microsoft.com/appx/manifest/uap/windows10");
+				Func<string, bool> addShowOnTile = (string s) => 
+				{
+					XmlElement ShowOn = AppxManifestXmlDocument.CreateElement("uap:ShowOn", "http://schemas.microsoft.com/appx/manifest/uap/windows10");
+					XmlAttribute Tile = AppxManifestXmlDocument.CreateAttribute("Tile");
+					Tile.Value = s;
+					ShowOn.Attributes.Append(Tile);
+					ShowNameOnTiles.AppendChild(ShowOn);
+					return true;
+				};
+
+				addShowOnTile("square150x150Logo");
+				if (bIsWideLogoUsed)
+				{
+					addShowOnTile("wide310x150Logo");
+				}
+				DefaultTile.AppendChild(ShowNameOnTiles);
+			}
+
+			return DefaultTile;
+		}
+
+		/// <summary>
+		/// Gather all information for the SplashScreen element of the manifest.
+		/// </summary>
+		private XmlNode GetSplashScreen(int ApplicationIndex)
+		{
+			XmlElement SplashScreen = AppxManifestXmlDocument!.CreateElement("uap:SplashScreen", "http://schemas.microsoft.com/appx/manifest/uap/windows10");
+
+			XmlAttribute BackgroundColor = CreateColorAttribute("BackgroundColor", "SplashScreenBackgroundColor", "Package.Applications.Application[" + ApplicationIndex + "].VisualElements.SplashScreen.BackgroundColor", "#000040");
+			SplashScreen.Attributes.Append(BackgroundColor);
+
+			XmlAttribute Image = AppxManifestXmlDocument.CreateAttribute("Image");
+			if (CopyAndReplaceBinaryIntermediate("SplashScreen.png"))
+			{
+				Image.Value = BuildResourceSubPath + "\\SplashScreen.png";
+				SplashScreen.Attributes.Append(Image);
+			}
+			else
+			{
+				Log.TraceError("Unable to stage splash screen image.");
+			}
+
+			return SplashScreen;
+		}
+
+        private XmlNode GetCapabilities(UnrealTargetPlatform TargetPlatform)
+        {
+            XmlElement Capabilities = AppxManifestXmlDocument!.CreateElement("Capabilities");
+
+            List<string>? CapabilityList = new List<string>();
+            List<string>? DeviceCapabilityList = new List<string>();
+            List<string>? UapCapabilityList = new List<string>();
+            List<string>? Uap2CapabilityList = new List<string>();
+
+            if (EngineIni!.GetArray(TargetSettings!, "CapabilityList", out CapabilityList))
+            {
+                foreach (string capName in CapabilityList)
+                {
+                    XmlElement CapabilityElement = AppxManifestXmlDocument.CreateElement("Capability");
+                    XmlAttribute Name = AppxManifestXmlDocument.CreateAttribute("Name");
+                    Name.Value = capName;
+                    CapabilityElement.Attributes.Append(Name);
+                    Capabilities.AppendChild(CapabilityElement);
+                }
+            }
+
+            if (EngineIni.GetArray(TargetSettings!, "UapCapabilityList", out UapCapabilityList))
+            {
+                foreach (string capName in UapCapabilityList)
+                {
+                    XmlElement CapabilityElement = AppxManifestXmlDocument.CreateElement("uap:Capability", "http://schemas.microsoft.com/appx/manifest/uap/windows10");
+                    XmlAttribute Name = AppxManifestXmlDocument.CreateAttribute("Name");
+                    Name.Value = capName;
+                    CapabilityElement.Attributes.Append(Name);
+                    Capabilities.AppendChild(CapabilityElement);
+                }
+            }
+
+            if (EngineIni.GetArray(TargetSettings!, "Uap2CapabilityList", out Uap2CapabilityList))
+            {
+                foreach (string capName in Uap2CapabilityList)
+                {
+                    XmlElement CapabilityElement = AppxManifestXmlDocument.CreateElement("uap2:Capability", "http://schemas.microsoft.com/appx/manifest/uap/windows10/2");
+                    XmlAttribute Name = AppxManifestXmlDocument.CreateAttribute("Name");
+                    Name.Value = capName;
+                    CapabilityElement.Attributes.Append(Name);
+                    Capabilities.AppendChild(CapabilityElement);
+                }
+            }
+
+            if (EngineIni.GetArray(TargetSettings!, "DeviceCapabilityList", out DeviceCapabilityList))
+            {
+                foreach (string capName in DeviceCapabilityList)
+                {
+                    XmlElement CapabilityElement = AppxManifestXmlDocument.CreateElement("DeviceCapability");
+                    XmlAttribute Name = AppxManifestXmlDocument.CreateAttribute("Name");
+                    Name.Value = capName;
+                    CapabilityElement.Attributes.Append(Name);
+                    Capabilities.AppendChild(CapabilityElement);
+                }
+            }
+
+			return Capabilities;
+        }
+
+		/// <summary>
+		/// Gather and create manifest for the package extension entries. There are multiple possible extension types
+		/// that can occur in any order and may individually be present or absent. A great deal of the complexity of this
+		/// function deals with correlating the old and new style INI entries across these order differences.
+		/// </summary>
+		private XmlNode? GetPackageExtensions()
+		{
+			XmlElement Extensions = AppxManifestXmlDocument!.CreateElement("Extensions");
+
+			foreach (var WinMD in WinMDReferences!)
+			{
+				XmlElement ExtensionElement = AppxManifestXmlDocument.CreateElement("Extension");
+				Extensions.AppendChild(ExtensionElement);
+
+				XmlAttribute CategoryAttribute = AppxManifestXmlDocument.CreateAttribute("Category");
+				CategoryAttribute.Value = "windows.activatableClass.inProcessServer";
+				ExtensionElement.Attributes.Append(CategoryAttribute);
+
+				XmlElement InProcessServerElement = AppxManifestXmlDocument.CreateElement("InProcessServer");
+				ExtensionElement.AppendChild(InProcessServerElement);
+
+				XmlElement PathElement = AppxManifestXmlDocument.CreateElement("Path");
+				InProcessServerElement.AppendChild(PathElement);
+				PathElement.InnerText = WinMD.PackageRelativeDllPath;
+
+				foreach (var WinMDType in WinMD.ActivatableTypes)
+				{
+					XmlElement ActivatableClassElement = AppxManifestXmlDocument.CreateElement("ActivatableClass");
+					InProcessServerElement.AppendChild(ActivatableClassElement);
+
+					XmlAttribute ActivatableClassIdAttribute = AppxManifestXmlDocument.CreateAttribute("ActivatableClassId");
+					ActivatableClassIdAttribute.Value = WinMDType.TypeName;
+					ActivatableClassElement.Attributes.Append(ActivatableClassIdAttribute);
+
+					XmlAttribute ThreadingModelAttribute = AppxManifestXmlDocument.CreateAttribute("ThreadingModel");
+					ThreadingModelAttribute.Value = WinMDType.ThreadingModelName;
+					ActivatableClassElement.Attributes.Append(ThreadingModelAttribute);
+				}
+			}
+
+			//@todo outOfProcessServer
+			//@todo proxyStub
+			//@todo windows.certificates
+
+			if (!Extensions.HasChildNodes)
+			{
+				return null;
+			}
+
+			return Extensions;
+		}
+
+		[SupportedOSPlatform("windows")]
+		private void ValidateAppxManifest(string ManifestPath)
+		{
+			System.Xml.Schema.XmlSchemaSet AppxSchema = new System.Xml.Schema.XmlSchemaSet();
+
+			// Validate against VS schemas if possible
+			DirectoryReference? SdkSchemaFolder = null;
+			DirectoryReference? VSSchemaFolder = null;
+			DirectoryReference? PhoneSchemaFolder = null;
+
+			// Limit to VS2015 compatible SDKs here - newer ones have incomplete schema sets
+			string SDKVersion = "";
+
+			DirectoryReference SDKRootFolder;
+			VersionNumber? version;
+			if (WindowsPlatform.TryGetWindowsSdkDir("Latest", out version, out SDKRootFolder!))
+			{
+				SDKVersion = version.ToString();
+			}
+
+			SdkSchemaFolder = DirectoryReference.Combine(SDKRootFolder, "Include", SDKVersion.ToString(), "winrt");
+			PhoneSchemaFolder = DirectoryReference.Combine(SDKRootFolder, "Extension SDKs", "WindowsMobile", SDKVersion.ToString(), "Include", "WinRT");
+
+			IEnumerable<DirectoryReference>? VSInstallDirs;
+			if (null != (VSInstallDirs = WindowsPlatform.TryGetVSInstallDirs(WindowsCompiler.VisualStudio2019)))
+			{
+				VSSchemaFolder = DirectoryReference.Combine(VSInstallDirs.First(), "Xml", "Schemas");
+			}
+			else if (null != (VSInstallDirs = WindowsPlatform.TryGetVSInstallDirs(WindowsCompiler.VisualStudio2022)))
+			{
+				VSSchemaFolder = DirectoryReference.Combine(VSInstallDirs.First(), "Xml", "Schemas");
+			}
+
+			string[] RequiredSchemas =
+			{
+				"AppxManifestTypes.xsd",
+				"UapManifestSchema.xsd",
+				"UapManifestSchema_v2.xsd",
+				"UapManifestSchema_v3.xsd",
+				"UapManifestSchema_v4.xsd",
+				"UapManifestSchema_v5.xsd",
+				"UapManifestSchema_v6.xsd",
+				"UapManifestSchema_v7.xsd",
+				"UapManifestSchema_v8.xsd",
+				"UapManifestSchema_v10.xsd",
+				"FoundationManifestSchema.xsd",
+				"AppxManifestSchema2010_v3.xsd",
+				"AppxManifestSchema2013_v2.xsd",
+				"AppxManifestSchema2014.xsd",
+				"AppxPhoneManifestSchema2014.xsd",
+				"DesktopManifestSchema_v2.xsd",
+				"DesktopManifestSchema_v3.xsd",
+				"DesktopManifestSchema_v4.xsd",
+				"DesktopManifestSchema_v5.xsd",
+				"DesktopManifestSchema_v6.xsd",
+				"IotManifestSchema_v2.xsd"
+			};
+
+			foreach (string SchemaName in RequiredSchemas)
+			{
+				FileReference? SchemaFile = null;
+
+				if (VSSchemaFolder != null)
+				{
+					SchemaFile = FileReference.Combine(VSSchemaFolder, SchemaName);
+				}
+
+				if ((SchemaFile == null || !FileReference.Exists(SchemaFile)) && SdkSchemaFolder != null)
+				{
+					SchemaFile = FileReference.Combine(SdkSchemaFolder, SchemaName);
+				}
+
+				if ((SchemaFile == null || !FileReference.Exists(SchemaFile)) && PhoneSchemaFolder != null)
+				{
+					SchemaFile = FileReference.Combine(PhoneSchemaFolder, SchemaName);
+				}
+
+				if (SchemaFile != null && FileReference.Exists(SchemaFile))
+				{
+					AppxSchema.Add(null, XmlReader.Create(SchemaFile.FullName));
+				}
+			}
+
+			try
+			{
+				AppxSchema.Compile();
+			}
+			catch (System.Xml.Schema.XmlSchemaException e)
+			{
+				string InvalidSchemaWarning =
+					"\r\n" +
+					"{0}({1}): {2}\r\n" +
+					"XML schema failed to compile; validation of the final AppxManifest.xml will be skipped.\r\n" +
+					"If your AppxManifest.xml is valid then this is harmless, but if it contains invalid content you may encounter packaging or deployment errors.\r\n" +
+					"Updating your Windows SDK and/or Visual Studio installation may correct the schema problems and simplify diagnosis of invalid content.\r\n";
+
+				Log.TraceWarning(InvalidSchemaWarning, e.SourceUri, e.LineNumber, e.Message);
+				return;
+			}
+
+			bool ValidationSucceeded = true;
+			XmlReaderSettings ReaderSettings = new XmlReaderSettings();
+			ReaderSettings.ValidationType = ValidationType.Schema;
+			ReaderSettings.Schemas = AppxSchema;
+			ReaderSettings.ValidationEventHandler += (source, args) =>
+			{
+				switch (args.Severity)
+				{
+					case System.Xml.Schema.XmlSeverityType.Error:
+						Log.TraceError(args.Message);
+						ValidationSucceeded = false;
+						break;
+
+					case System.Xml.Schema.XmlSeverityType.Warning:
+						Log.TraceWarning(args.Message);
+						break;
+
+					default:
+						break;
+				}
+			};
+
+			using (XmlReader ValidatingReader = XmlReader.Create(ManifestPath, ReaderSettings))
+			{
+				while (ValidatingReader.Read())
+				{
+					// No-op, just reading to end to force validation.
+				}
+			}
+
+			if (!ValidationSucceeded)
+			{
+				throw new BuildException("Generated AppxManifest ({0}) is invalid.  See log for details and check your HoloLens Project Settings.", ManifestPath);
+			}
+		}
+	};
+}
diff --git a/Engine/Platforms/Hololens/Source/Programs/UnrealBuildTool/UEBuildHoloLens.cs b/Engine/Platforms/Hololens/Source/Programs/UnrealBuildTool/UEBuildHoloLens.cs
new file mode 100644
--- /dev/null
+++ b/Engine/Platforms/Hololens/Source/Programs/UnrealBuildTool/UEBuildHoloLens.cs
@@ -0,0 +1,762 @@
+// Copyright Epic Games, Inc. All Rights Reserved.
+
+using System;
+using System.Collections.Generic;
+using System.Text;
+using System.Diagnostics;
+using System.IO;
+using System.Linq;
+using Microsoft.Win32;
+using EpicGames.Core;
+using System.Text.RegularExpressions;
+using System.Runtime.Versioning;
+
+namespace UnrealBuildTool
+{
+	/// <summary>
+	/// HoloLens-specific target settings
+	/// </summary>
+	public class HoloLensTargetRules
+	{
+		/// <summary>
+		/// Version of the compiler toolchain to use on HoloLens. A value of "default" will be changed to a specific version at UBT startup.
+		/// </summary>
+		[ConfigFile(ConfigHierarchyType.Engine, "/Script/HoloLensPlatformEditor.HoloLensTargetSettings", "CompilerVersion")]
+		[XmlConfigFile(Category = "HoloLensPlatform")]
+		[CommandLine("-2019", Value = nameof(WindowsCompiler.VisualStudio2019))]
+		[CommandLine("-2022", Value = nameof(WindowsCompiler.VisualStudio2022))]
+		public WindowsCompiler Compiler = WindowsCompiler.Default;
+
+		/// <summary>
+		/// Architecture of Target.
+		/// </summary>
+		public WindowsArchitecture Architecture
+		{
+			get;
+		}
+
+		/// <summary>
+		/// Enable PIX debugging (automatically disabled in Shipping and Test configs)
+		/// </summary>
+		[ConfigFile(ConfigHierarchyType.Engine, "/Script/HoloLensPlatformEditor.HoloLensTargetSettings", "bEnablePIXProfiling")]
+		public bool bPixProfilingEnabled = true;
+
+		/// <summary>
+		/// Version of the compiler toolchain to use on HoloLens. A value of "default" will be changed to a specific version at UBT startup.
+		/// </summary>
+		[ConfigFile(ConfigHierarchyType.Engine, "/Script/HoloLensPlatformEditor.HoloLensTargetSettings", "bBuildForRetailWindowsStore")]
+		public bool bBuildForRetailWindowsStore = false;
+
+		/// <summary>
+		/// Contains the specific version of the Windows 10 SDK that we will build against. If empty, it will be "Latest"
+		/// </summary>
+		[ConfigFile(ConfigHierarchyType.Engine, "/Script/HoloLensPlatformEditor.HoloLensTargetSettings", "Windows10SDKVersion")]
+		public string? Win10SDKVersionString = null;
+
+		internal Version? Win10SDKVersion = null;
+
+        /// <summary>
+        /// Automatically increment the project version after each build.
+        /// </summary>
+        [ConfigFile(ConfigHierarchyType.Engine, "/Script/HoloLensPlatformEditor.HoloLensTargetSettings", "bAutoIncrementVersion")]
+        public bool bAutoIncrementVersion = false;
+
+		/// <summary>
+		/// Whether to run native code analysis at compile time, producing a nativecodeanalysis.xml file for every compiled source file
+		/// </summary>
+		[ConfigFile(ConfigHierarchyType.Engine, "/Script/HoloLensPlatformEditor.HoloLensTargetSettings", "bRunNativeCodeAnalysis")]
+		public bool bRunNativeCodeAnalysis = false;
+		/// <summary>
+		/// A project relative path for a custom native code analysis ruleset xml file
+		/// </summary>
+		[ConfigFile(ConfigHierarchyType.Engine, "/Script/HoloLensPlatformEditor.HoloLensTargetSettings", "NativeCodeAnalysisRuleset")]
+		public string? NativeCodeAnalysisRuleset = null;
+		
+		/// <summary>
+		/// Constructor
+		/// </summary>
+		/// <param name="Info">Target information</param>
+		public HoloLensTargetRules(TargetInfo Info)
+		{
+			if (Info.Platform == UnrealTargetPlatform.HoloLens && !String.IsNullOrEmpty(Info.Architecture))
+			{
+				Architecture = (WindowsArchitecture)Enum.Parse(typeof(WindowsArchitecture), Info.Architecture, true);
+			}
+		}
+    }
+
+	/// <summary>
+	/// Read-only wrapper for HoloLens-specific target settings
+	/// </summary>
+	public class ReadOnlyHoloLensTargetRules
+	{
+		/// <summary>
+		/// The private mutable settings object
+		/// </summary>
+		private HoloLensTargetRules Inner;
+
+		/// <summary>
+		/// Constructor
+		/// </summary>
+		/// <param name="Inner">The settings object to wrap</param>
+		public ReadOnlyHoloLensTargetRules(HoloLensTargetRules Inner)
+		{
+			this.Inner = Inner;
+		}
+
+		/// <summary>
+		/// Accessors for fields on the inner TargetRules instance
+		/// </summary>
+		#region Read-only accessor properties 
+#pragma warning disable CS1591
+		public WindowsCompiler Compiler
+		{
+			get { return Inner.Compiler; }
+		}
+
+		public WindowsArchitecture Architecture
+		{
+			get { return Inner.Architecture; }
+		}
+
+		public bool bPixProfilingEnabled
+		{
+			get { return Inner.bPixProfilingEnabled; }
+		}
+
+		public bool bBuildForRetailWindowsStore
+		{
+			get { return Inner.bBuildForRetailWindowsStore; }
+		}
+
+		public Version? Win10SDKVersion
+		{
+			get { return Inner.Win10SDKVersion; }
+		}
+		public string? Win10SDKVersionString
+		{
+			get { return Inner.Win10SDKVersionString; }
+		}
+		
+		public bool bRunNativeCodeAnalysis
+		{
+			get { return Inner.bRunNativeCodeAnalysis; }
+		}
+		public string? NativeCodeAnalysisRuleset
+		{
+			get { return Inner.NativeCodeAnalysisRuleset; }
+		}
+#pragma warning restore CS1591
+		#endregion
+	}
+
+	[SupportedOSPlatform("windows")]
+	class HoloLensPlatform : WindowsPlatform
+	{
+		public static readonly Version MinimumSDKVersionRecommended = new Version(10, 0, 17763, 0);
+		public static readonly Version MaximumSDKVersionTested = new Version(10, 0, 18362, int.MaxValue);
+
+		public HoloLensPlatform(MicrosoftPlatformSDK InSDK) 
+			: base(UnrealTargetPlatform.HoloLens, InSDK)
+		{
+		}
+
+		public override void ValidateTarget(TargetRules Target)
+		{
+			// WindowsTargetRules are reused for HoloLens, so that build modules can keep the model that reuses "windows" configs for most cases
+			// That means overriding those settings here that need to be adjusted for HoloLens
+
+			// Compiler version and pix flags must be reloaded from the HoloLens hive
+
+			// Currently BP-only projects don't load build-related settings from their remote ini when building UnrealGame.exe
+			// (see TargetRules.cs, where the possibly-null project directory is passed to ConfigCache.ReadSettings).
+			// It's important for HoloLens that we *do* use the project-specific settings when building (VS 2017 vs 2015 and
+			// retail Windows Store are both examples).  Possibly this should be done on all platforms?  But in the interest
+			// of not changing behavior on other platforms I'm limiting the scope.
+
+			DirectoryReference? IniDirRef = DirectoryReference.FromFile(Target.ProjectFile);
+			if (IniDirRef == null && !string.IsNullOrEmpty(UnrealBuildTool.GetRemoteIniPath()))
+			{
+				IniDirRef = new DirectoryReference(UnrealBuildTool.GetRemoteIniPath()!);
+			}
+
+			// Stash the current compiler choice (accounts for command line) in case ReadSettings reverts it to default
+			WindowsCompiler CompilerBeforeReadSettings = Target.HoloLensPlatform.Compiler;
+
+			ConfigCache.ReadSettings(IniDirRef, Platform, Target.HoloLensPlatform);
+
+			if (Target.HoloLensPlatform.Compiler == WindowsCompiler.Default)
+			{
+				if (CompilerBeforeReadSettings != WindowsCompiler.Default)
+				{
+					// Previous setting was more specific, use that
+					Target.HoloLensPlatform.Compiler = CompilerBeforeReadSettings;
+				}
+				else
+				{
+					Target.HoloLensPlatform.Compiler = WindowsPlatform.GetDefaultCompiler(Target.ProjectFile, Target.HoloLensPlatform.Architecture);
+				}
+			}
+
+			if (!Target.bGenerateProjectFiles)
+			{
+				Log.TraceInformationOnce("Using {0} architecture for deploying to HoloLens device", Target.HoloLensPlatform.Architecture);
+			}
+
+			Target.WindowsPlatform.Compiler = Target.HoloLensPlatform.Compiler;
+			Target.WindowsPlatform.Architecture = Target.HoloLensPlatform.Architecture;
+			Target.WindowsPlatform.bPixProfilingEnabled = Target.HoloLensPlatform.bPixProfilingEnabled;
+			Target.WindowsPlatform.bUseWindowsSDK10 = true;
+			Target.WindowsPlatform.bStrictConformanceMode = false;
+
+			Target.bDeployAfterCompile = true;
+			Target.bCompileNvCloth = false;      // requires CUDA
+
+			// Disable Simplygon support if compiling against the NULL RHI.
+			if (Target.GlobalDefinitions.Contains("USE_NULL_RHI=1"))
+			{
+				Target.bCompileSpeedTree = false;
+			}
+
+			// Use shipping binaries to avoid dependency on nvToolsExt which fails WACK.
+			if (Target.Configuration == UnrealTargetConfiguration.Shipping)
+			{
+				Target.bUseShippingPhysXLibraries = true;
+			}
+
+			// Be resilient to SDKs being uninstalled but still referenced in the INI file
+			VersionNumber? SelectedWindowsSdkVersion;
+			DirectoryReference? SelectedWindowsSdkDir;
+			if (WindowsPlatform.TryGetWindowsSdkDir(Target.HoloLensPlatform.Win10SDKVersionString, out SelectedWindowsSdkVersion, out SelectedWindowsSdkDir))
+			{
+				Target.WindowsPlatform.WindowsSdkVersion = Target.HoloLensPlatform.Win10SDKVersionString;
+			}
+
+			base.ValidateTarget(Target);
+
+			// Windows 10 SDK version
+			// Auto-detect latest compatible by default (recommended), allow for explicit override if necessary
+			// Validate that the SDK isn't too old, and that the combination of VS and SDK is supported.
+
+			Target.HoloLensPlatform.Win10SDKVersion = new Version(Target.WindowsPlatform.Environment!.WindowsSdkVersion.ToString());
+
+			if(!Target.bGenerateProjectFiles)
+			{
+				Log.TraceInformationOnce("Building using Windows SDK version {0} for HoloLens and compiler {1} {2}", Target.HoloLensPlatform.Win10SDKVersion, Target.HoloLensPlatform.Compiler, Target.WindowsPlatform.CompilerVersion);
+
+				if (Target.HoloLensPlatform.Win10SDKVersion < MinimumSDKVersionRecommended)
+				{
+					Log.TraceWarning("Your Windows SDK version {0} is older than the minimum recommended version ({1}) for HoloLens.  Consider upgrading.", Target.HoloLensPlatform.Win10SDKVersion, MinimumSDKVersionRecommended);
+				}
+				else if (Target.HoloLensPlatform.Win10SDKVersion > MaximumSDKVersionTested)
+				{
+					Log.TraceInformationOnce("Your Windows SDK version ({0}) for HoloLens is newer than the highest tested with this version of UBT ({1}).  This is probably fine, but if you encounter issues consider using an earlier SDK.", Target.HoloLensPlatform.Win10SDKVersion, MaximumSDKVersionTested);
+				}
+			}
+
+			HoloLensExports.InitWindowsSdkToolPath(Target.HoloLensPlatform.Win10SDKVersion.ToString());
+
+			// ISPC currently doesn't support Windows-AArch64
+			Target.bCompileISPC = false;
+		}
+
+		/// <summary>
+		/// Gets the default HoloLens architecture
+		/// </summary>
+		/// <param name="ProjectFile">The project being built</param>
+		/// <returns>The default architecture</returns>
+		public override string GetDefaultArchitecture(FileReference? ProjectFile)
+		{
+			return WindowsArchitecture.x64.ToString();
+		}
+
+		/// <summary>
+		/// Get the extension to use for the given binary type
+		/// </summary>
+		/// <param name="InBinaryType"> The binrary type being built</param>
+		/// <returns>string	The binary extenstion (ie 'exe' or 'dll')</returns>
+		public override string GetBinaryExtension(UEBuildBinaryType InBinaryType)
+		{
+			switch (InBinaryType)
+			{
+				case UEBuildBinaryType.Object:
+					return ".obj";
+				case UEBuildBinaryType.PrecompiledHeader:
+					return ".pch";
+			}
+			return base.GetBinaryExtension(InBinaryType);
+		}
+
+		internal static DirectoryReference? GetCppCXMetadataLocation(WindowsCompiler Compiler, string CompilerVersion, WindowsArchitecture Architecture)
+		{
+			VersionNumber? SelectedToolChainVersion;
+			DirectoryReference? SelectedToolChainDir;
+			DirectoryReference? SelectedRedistDir;
+			if (!WindowsPlatform.TryGetToolChainDir(Compiler, CompilerVersion, Architecture, out SelectedToolChainVersion, out SelectedToolChainDir, out SelectedRedistDir))
+			{
+				return null;
+			}
+
+			return GetCppCXMetadataLocation(Compiler, SelectedToolChainDir);
+		}
+
+		public static DirectoryReference? GetCppCXMetadataLocation(WindowsCompiler Compiler, DirectoryReference SelectedToolChainDir)
+		{
+			if (Compiler.IsMSVC())
+			{
+				return DirectoryReference.Combine(SelectedToolChainDir, "lib", "x86", "Store", "references");
+			}
+			else
+			{
+				return null;
+			}
+		}
+
+
+		private static Version FindLatestVersionDirectory(string InDirectory, Version? NoLaterThan)
+		{
+			Version LatestVersion = new Version(0, 0, 0, 0);
+			if (Directory.Exists(InDirectory))
+			{
+				string[] VersionDirectories = Directory.GetDirectories(InDirectory);
+				foreach (string Dir in VersionDirectories)
+				{
+					string VersionString = Path.GetFileName(Dir);
+					Version? FoundVersion;
+					if (Version.TryParse(VersionString, out FoundVersion) && FoundVersion > LatestVersion)
+					{
+						if (NoLaterThan == null || FoundVersion <= NoLaterThan)
+						{
+							LatestVersion = FoundVersion;
+						}
+					}
+				}
+			}
+			return LatestVersion;
+		}
+
+		internal static string GetLatestMetadataPathForApiContract(string ApiContract, WindowsCompiler Compiler)
+		{
+			DirectoryReference? SDKFolder;
+			VersionNumber? SDKVersion;
+			if (!WindowsPlatform.TryGetWindowsSdkDir("Latest", out SDKVersion, out SDKFolder))
+			{
+				return string.Empty;
+			}
+
+			DirectoryReference ReferenceDir = DirectoryReference.Combine(SDKFolder, "References");
+			if (DirectoryReference.Exists(ReferenceDir))
+			{
+				// Prefer a contract from a suitable SDK-versioned subdir of the references folder when available (starts with 15063 SDK)
+				DirectoryReference SDKVersionedReferenceDir = DirectoryReference.Combine(ReferenceDir, SDKVersion.ToString());
+				DirectoryReference ContractDir = DirectoryReference.Combine(SDKVersionedReferenceDir, ApiContract);
+				Version? ContractLatestVersion = null;
+				FileReference? MetadataFileRef = null;
+				if (DirectoryReference.Exists(ContractDir))
+				{
+					// Note: contract versions don't line up with Windows SDK versions (they're numbered independently as 1.0.0.0, 2.0.0.0, etc.)
+					ContractLatestVersion = FindLatestVersionDirectory(ContractDir.FullName, null);
+					MetadataFileRef = FileReference.Combine(ContractDir, ContractLatestVersion.ToString(), ApiContract + ".winmd");
+				}
+
+				// Retry in unversioned references dir if we failed above.
+				if (MetadataFileRef == null || !FileReference.Exists(MetadataFileRef))
+				{
+					ContractDir = DirectoryReference.Combine(ReferenceDir, ApiContract);
+					if (DirectoryReference.Exists(ContractDir))
+					{
+						ContractLatestVersion = FindLatestVersionDirectory(ContractDir.FullName, null);
+						MetadataFileRef = FileReference.Combine(ContractDir, ContractLatestVersion.ToString(), ApiContract + ".winmd");
+					}
+				}
+				if (MetadataFileRef != null && FileReference.Exists(MetadataFileRef))
+				{
+					return MetadataFileRef.FullName;
+				}
+			}
+
+			return string.Empty;
+		}
+
+
+		/// <summary>
+		/// Modify the rules for a newly created module, where the target is a different host platform.
+		/// This is not required - but allows for hiding details of a particular platform.
+		/// </summary>
+		/// <param name="ModuleName">The name of the module</param>
+		/// <param name="Rules">The module rules</param>
+		/// <param name="Target">The target being build</param>
+		public override void ModifyModuleRulesForOtherPlatform(string ModuleName, ModuleRules Rules, ReadOnlyTargetRules Target)
+		{
+			// don't do any target platform stuff if SDK is not available
+			if (!UEBuildPlatform.IsPlatformAvailableForTarget(Platform, Target))
+			{
+				return;
+			}
+			
+			if (Target.Platform == UnrealTargetPlatform.Win64)
+			{
+				if (!Target.bBuildRequiresCookedData)
+				{
+					if (ModuleName == "Engine")
+					{
+						if (Target.bBuildDeveloperTools)
+						{
+							Rules.DynamicallyLoadedModuleNames.Add("HoloLensTargetPlatform");
+						}
+					}
+				}
+
+				// allow standalone tools to use targetplatform modules, without needing Engine
+				if (ModuleName == "TargetPlatform")
+				{
+					if (Target.bForceBuildTargetPlatforms)
+					{
+						Rules.DynamicallyLoadedModuleNames.Add("HoloLensTargetPlatform");
+					}
+				}
+
+				if (ModuleName == "UnrealEd")
+				{
+					Rules.DynamicallyLoadedModuleNames.Add("HoloLensPlatformEditor");
+				}
+			}
+			
+// This code has been removed because it causes a full rebuild after generating project files (since response files are overwritten with different defines).
+#if false
+			if (Target.Platform == UnrealTargetPlatform.Win64)
+			{
+				if (ProjectFileGenerator.bGenerateProjectFiles)
+				{
+					// Use latest SDK for Intellisense purposes
+					WindowsCompiler CompilerForSdkRestriction = Target.HoloLensPlatform.Compiler != WindowsCompiler.Default ? Target.HoloLensPlatform.Compiler : Target.WindowsPlatform.Compiler;
+					if (CompilerForSdkRestriction != WindowsCompiler.Default)
+					{
+						Version OutWin10SDKVersion;
+						DirectoryReference OutSdkDir;
+						if(WindowsExports.TryGetWindowsSdkDir(Target.HoloLensPlatform.Win10SDKVersionString, out OutWin10SDKVersion, out OutSdkDir))
+						{
+							Rules.PublicDefinitions.Add(string.Format("WIN10_SDK_VERSION={0}", OutWin10SDKVersion.Build));
+						}
+					}
+				}
+			}
+#endif
+		}
+
+		/// <summary>
+		/// Deploys the given target
+		/// </summary>
+		/// <param name="Receipt">Information about the target being deployed</param>
+		public override void Deploy(TargetReceipt Receipt)
+		{
+			new HoloLensDeploy().PrepTargetForDeployment(Receipt);
+		}
+
+		/// <summary>
+		/// Modify the rules for a newly created module, in a target that's being built for this platform.
+		/// This is not required - but allows for hiding details of a particular platform.
+		/// </summary>
+		/// <param name="ModuleName">The name of the module</param>
+		/// <param name="Rules">The module rules</param>
+		/// <param name="Target">The target being build</param>
+		public override void ModifyModuleRulesForActivePlatform(string ModuleName, ModuleRules Rules, ReadOnlyTargetRules Target)
+		{
+			if (ModuleName == "Core")
+			{
+				//Rules.PrivateDependencyModuleNames.Add("HoloLensSDK");
+			}
+			else if (ModuleName == "Engine")
+			{
+				Rules.PrivateDependencyModuleNames.Add("zlib");
+				Rules.PrivateDependencyModuleNames.Add("UElibPNG");
+				Rules.PublicDependencyModuleNames.Add("UEOgg");
+				Rules.PublicDependencyModuleNames.Add("Vorbis");
+			}
+			else if (ModuleName == "D3D11RHI")
+			{
+				Rules.PublicDefinitions.Add("D3D11_WITH_DWMAPI=0");
+				Rules.PublicDefinitions.Add("WITH_DX_PERF=0");
+				Rules.PrivateDependencyModuleNames.Add("WinPixEventRuntime");
+			}
+			else if (ModuleName == "D3D12RHI")
+			{
+				// To enable platform specific D3D12 RHI Types
+				Rules.PrivateIncludePaths.Add("Runtime/D3D12RHI/Private/HoloLens");
+				Rules.PrivateDependencyModuleNames.Add("WinPixEventRuntime");
+			}
+			else if (ModuleName == "DX11")
+			{
+				// Clear out all the Windows include paths and libraries...
+				// The HoloLensSDK module handles proper paths and libs for HoloLens.
+				// However, the D3D11RHI module will include the DX11 module.
+				Rules.PublicIncludePaths.Clear();
+				Rules.InternalncludePaths.Clear();
+				Rules.PublicSystemLibraryPaths.Clear();
+				Rules.PublicSystemLibraries.Clear();
+				Rules.PublicAdditionalLibraries.Clear();
+				Rules.PublicAdditionalLibraries.Remove("X3DAudio.lib");
+				Rules.PublicAdditionalLibraries.Remove("XAPOFX.lib");
+				Rules.PrivateDependencyModuleNames.Add("WinPixEventRuntime");
+			}
+			else if (ModuleName == "DX11Audio")
+			{
+				Rules.PublicAdditionalLibraries.Remove("X3DAudio.lib");
+				Rules.PublicAdditionalLibraries.Remove("XAPOFX.lib");
+			}
+		}
+
+		internal static void ExpandWinMDReferences(ReadOnlyTargetRules Target, string SDKFolder, string SDKVersion, ref List<string> WinMDReferences)
+		{
+			// Code below will fail when not using the Win10 SDK.  Early out to avoid warning spam.
+			if (!Target.WindowsPlatform.bUseWindowsSDK10)
+			{
+				return;
+			}
+
+			if (WinMDReferences.Count > 0)
+			{
+				// Allow bringing in Windows SDK contracts just by naming the contract
+				// These are files that look like References/10.0.98765.0/AMadeUpWindowsApiContract/5.0.0.0/AMadeUpWindowsApiContract.winmd
+				List<string> ExpandedWinMDReferences = new List<string>();
+
+				// The first few releases of the Windows 10 SDK didn't put the SDK version in the reference path
+				string ReferenceRoot = Path.Combine(SDKFolder, "References");
+				string VersionedReferenceRoot = Path.Combine(ReferenceRoot, SDKVersion);
+				if (Directory.Exists(VersionedReferenceRoot))
+				{
+					ReferenceRoot = VersionedReferenceRoot;
+				}
+
+				foreach (string WinMDRef in WinMDReferences)
+				{
+					if (File.Exists(WinMDRef))
+					{
+						// Already a valid path
+						ExpandedWinMDReferences.Add(WinMDRef);
+					}
+					else
+					{
+						string ContractFolder = Path.Combine(ReferenceRoot, WinMDRef);
+
+						Version ContractVersion = FindLatestVersionDirectory(ContractFolder, null);
+						string ExpandedWinMDRef = Path.Combine(ContractFolder, ContractVersion.ToString(), WinMDRef + ".winmd");
+						if (File.Exists(ExpandedWinMDRef))
+						{
+							ExpandedWinMDReferences.Add(ExpandedWinMDRef);
+						}
+						else
+						{
+							Log.TraceWarning("Unable to resolve location for HoloLens WinMD api contract {0}, file {1}", WinMDRef, ExpandedWinMDRef);
+						}
+					}
+				}
+
+				WinMDReferences = ExpandedWinMDReferences;
+			}
+		}
+
+		/// <summary>
+		/// Setup the target environment for building
+		/// </summary>
+		/// <param name="Target">Settings for the target being compiled</param>
+		/// <param name="CompileEnvironment">The compile environment for this target</param>
+		/// <param name="LinkEnvironment">The link environment for this target</param>
+		public override void SetUpEnvironment(ReadOnlyTargetRules Target, CppCompileEnvironment CompileEnvironment, LinkEnvironment LinkEnvironment)
+		{
+			// Add Win10 SDK pieces - moved here since it allows better control over SDK version
+			string? Win10SDKRoot = Target.WindowsPlatform.WindowsSdkDir;
+
+			// Reference (WinMD) paths
+			// Only Foundation and Universal are referenced by default.  
+			List<string> AlwaysReferenceContracts = new List<string>();
+			AlwaysReferenceContracts.Add("Windows.Foundation.FoundationContract");
+			AlwaysReferenceContracts.Add("Windows.Foundation.UniversalApiContract");
+			ExpandWinMDReferences(Target, Win10SDKRoot!, Target.HoloLensPlatform.Win10SDKVersion!.ToString(), ref AlwaysReferenceContracts);
+
+			StringBuilder WinMDReferenceArguments = new StringBuilder();
+			foreach (string WinMDReference in AlwaysReferenceContracts)
+			{
+				WinMDReferenceArguments.AppendFormat(@" /FU""{0}""", WinMDReference);
+			}
+			CompileEnvironment.AdditionalArguments += WinMDReferenceArguments.ToString().Trim();
+
+			CompileEnvironment.Definitions.Add("EXCEPTIONS_DISABLED=0");
+
+			CompileEnvironment.Definitions.Add("_WIN32_WINNT=0x0A00");
+			CompileEnvironment.Definitions.Add("WINVER=0x0A00");
+
+			CompileEnvironment.Definitions.Add("PLATFORM_HOLOLENS=1");
+			CompileEnvironment.Definitions.Add("HOLOLENS=1");
+
+			CompileEnvironment.Definitions.Add("WINAPI_FAMILY=WINAPI_FAMILY_APP");
+			CompileEnvironment.Definitions.Add("PLATFORM_MICROSOFT=1");
+
+			// No D3DX on HoloLens!
+			CompileEnvironment.Definitions.Add("NO_D3DX_LIBS=1");
+
+			if (Target.HoloLensPlatform.bBuildForRetailWindowsStore)
+			{
+				CompileEnvironment.Definitions.Add("USING_RETAIL_WINDOWS_STORE=1");
+			}
+			else
+			{
+				CompileEnvironment.Definitions.Add("USING_RETAIL_WINDOWS_STORE=0");
+			}
+
+			CompileEnvironment.Definitions.Add("WITH_D3D12_RHI=0");
+
+			LinkEnvironment.AdditionalArguments += "/NODEFAULTLIB";
+			//CompileEnvironment.AdditionalArguments += " /showIncludes";
+
+			LinkEnvironment.SystemLibraries.Add("windowsapp.lib");
+
+			CompileEnvironment.Definitions.Add(string.Format("WIN10_SDK_VERSION={0}", Target.HoloLensPlatform.Win10SDKVersion.Build));
+
+			LinkEnvironment.SystemLibraries.Add("dloadhelper.lib");
+			LinkEnvironment.SystemLibraries.Add("ws2_32.lib");
+
+			if (CompileEnvironment.bUseDebugCRT)
+			{
+				LinkEnvironment.SystemLibraries.Add("vccorlibd.lib");
+				LinkEnvironment.SystemLibraries.Add("ucrtd.lib");
+				LinkEnvironment.SystemLibraries.Add("vcruntimed.lib");
+				LinkEnvironment.SystemLibraries.Add("msvcrtd.lib");
+				LinkEnvironment.SystemLibraries.Add("msvcprtd.lib");
+			}
+			else
+			{
+				LinkEnvironment.SystemLibraries.Add("vccorlib.lib");
+				LinkEnvironment.SystemLibraries.Add("ucrt.lib");
+				LinkEnvironment.SystemLibraries.Add("vcruntime.lib");
+				LinkEnvironment.SystemLibraries.Add("msvcrt.lib");
+				LinkEnvironment.SystemLibraries.Add("msvcprt.lib");
+			}
+			LinkEnvironment.SystemLibraries.Add("legacy_stdio_wide_specifiers.lib");
+			LinkEnvironment.SystemLibraries.Add("uuid.lib");
+
+			// Static CRT not supported for HoloLens
+			CompileEnvironment.bUseStaticCRT = false;
+			CompileEnvironment.bEnableExceptions = true;
+		}
+
+		/// <summary>
+		/// Setup the configuration environment for building
+		/// </summary>
+		/// <param name="Target"> The target being built</param>
+		/// <param name="GlobalCompileEnvironment">The global compile environment</param>
+		/// <param name="GlobalLinkEnvironment">The global link environment</param>
+		public override void SetUpConfigurationEnvironment(ReadOnlyTargetRules Target, CppCompileEnvironment GlobalCompileEnvironment, LinkEnvironment GlobalLinkEnvironment)
+		{
+			// Determine the C++ compile/link configuration based on the Unreal configuration.
+
+			if (GlobalCompileEnvironment.bUseDebugCRT)
+			{
+				GlobalCompileEnvironment.Definitions.Add("_DEBUG=1"); // the engine doesn't use this, but lots of 3rd party stuff does
+			}
+			else
+			{
+				GlobalCompileEnvironment.Definitions.Add("NDEBUG=1"); // the engine doesn't use this, but lots of 3rd party stuff does
+			}
+
+			//CppConfiguration CompileConfiguration;
+			UnrealTargetConfiguration CheckConfig = Target.Configuration;
+			switch (CheckConfig)
+			{
+				default:
+				case UnrealTargetConfiguration.Debug:
+					GlobalCompileEnvironment.Definitions.Add("UE_BUILD_DEBUG=1");
+					break;
+				case UnrealTargetConfiguration.DebugGame:
+				// Default to Development; can be overriden by individual modules.
+				case UnrealTargetConfiguration.Development:
+					GlobalCompileEnvironment.Definitions.Add("UE_BUILD_DEVELOPMENT=1");
+					break;
+				case UnrealTargetConfiguration.Shipping:
+					GlobalCompileEnvironment.Definitions.Add("UE_BUILD_SHIPPING=1");
+					break;
+				case UnrealTargetConfiguration.Test:
+					GlobalCompileEnvironment.Definitions.Add("UE_BUILD_TEST=1");
+					break;
+			}
+
+			// Create debug info based on the heuristics specified by the user.
+			GlobalCompileEnvironment.bCreateDebugInfo =
+				!Target.bDisableDebugInfo && ShouldCreateDebugInfo(Target);
+
+			// NOTE: Even when debug info is turned off, we currently force the linker to generate debug info
+			//	   anyway on Visual C++ platforms.  This will cause a PDB file to be generated with symbols
+			//	   for most of the classes and function/method names, so that crashes still yield somewhat
+			//	   useful call stacks, even though compiler-generate debug info may be disabled.  This gives
+			//	   us much of the build-time savings of fully-disabled debug info, without giving up call
+			//	   data completely.
+			GlobalLinkEnvironment.bCreateDebugInfo = true;
+		}
+
+		/// <summary>
+		/// Whether this platform should create debug information or not
+		/// </summary>
+		/// <param name="Target">The target being built</param>
+		/// <returns>bool	true if debug info should be generated, false if not</returns>
+		public override bool ShouldCreateDebugInfo(ReadOnlyTargetRules Target)
+		{
+			switch (Target.Configuration)
+			{
+				case UnrealTargetConfiguration.Development:
+				case UnrealTargetConfiguration.Shipping:
+				case UnrealTargetConfiguration.Test:
+					return !Target.bOmitPCDebugInfoInDevelopment;
+				case UnrealTargetConfiguration.DebugGame:
+				case UnrealTargetConfiguration.Debug:
+				default:
+					return true;
+			};
+		}
+
+		/// <summary>
+		/// Creates a toolchain instance for the given platform.
+		/// </summary>
+		/// <param name="Target">The target being built</param>
+		/// <returns>New toolchain instance.</returns>
+		public override UEToolChain CreateToolChain(ReadOnlyTargetRules Target)
+		{
+			return new HoloLensToolChain(Target);
+		}
+
+		/// <summary>
+		/// Check for the build requirement due to platform requirements
+		/// return true if the project requires a build
+		/// </summary>
+		public override bool RequiresBuild(UnrealTargetPlatform Platform, DirectoryReference ProjectDirectoryName)
+		{
+			// Hololens requires UBT run before packaging so that HoloLensDeploy PrepTargetForDeployment will create the Multi.target file which is used in hololens packaging.
+			return true;
+		}
+	}
+
+
+
+	class HoloLensPlatformFactory : UEBuildPlatformFactory
+	{
+		public override UnrealTargetPlatform TargetPlatform
+		{
+			get { return UnrealTargetPlatform.HoloLens; }
+		}
+
+		/// <summary>
+		/// Register the platform with the UEBuildPlatform class
+		/// </summary>
+		public override void RegisterBuildPlatforms()
+		{
+			if (OperatingSystem.IsWindows())
+			{
+				// for GetValidSoftwareVersionRange reasons we probably want a HoloLensePlatformSDK class
+				MicrosoftPlatformSDK SDK = new MicrosoftPlatformSDK();
+
+				UEBuildPlatform.RegisterBuildPlatform(new HoloLensPlatform(SDK));
+				UEBuildPlatform.RegisterPlatformWithGroup(UnrealTargetPlatform.HoloLens, UnrealPlatformGroup.Microsoft);
+				UEBuildPlatform.RegisterPlatformWithGroup(UnrealTargetPlatform.HoloLens, UnrealPlatformGroup.HoloLens);
+			}
+		}
+	}
+}
diff --git a/Engine/Platforms/Hololens/Source/Programs/UnrealBuildTool/WinMDRegistrationInfo.cs b/Engine/Platforms/Hololens/Source/Programs/UnrealBuildTool/WinMDRegistrationInfo.cs
new file mode 100644
--- /dev/null
+++ b/Engine/Platforms/Hololens/Source/Programs/UnrealBuildTool/WinMDRegistrationInfo.cs
@@ -0,0 +1,191 @@
+// Copyright Epic Games, Inc. All Rights Reserved.
+
+using System;
+using System.Collections.Generic;
+using System.IO;
+using System.Linq;
+using System.Reflection;
+using System.Runtime.InteropServices;
+using System.Runtime.Versioning;
+using EpicGames.Core;
+
+namespace UnrealBuildTool
+{
+	/// <summary>
+	/// WinMD registration helper
+	/// </summary>
+	public class WinMDRegistrationInfo
+	{
+		/// <summary>
+		/// WinMD type info
+		/// </summary>
+		public class ActivatableType
+		{
+			/// <summary>
+			/// WinMD type info constructor
+			/// </summary>
+			/// <param name="InTypeName"></param>
+			/// <param name="InThreadingModelName"></param>
+			public ActivatableType(string InTypeName, string InThreadingModelName)
+			{
+				TypeName = InTypeName;
+				ThreadingModelName = InThreadingModelName;
+			}
+
+			/// <summary>
+			/// Type Name
+			/// </summary>
+			public string TypeName { get; private set; }
+
+			/// <summary>
+			/// Threading model
+			/// </summary>
+			public string ThreadingModelName { get; private set; }
+		}
+
+		private static Version FindLatestVersionDirectory(string InDirectory, Version? NoLaterThan)
+		{
+			Version LatestVersion = new Version(0, 0, 0, 0);
+			if (Directory.Exists(InDirectory))
+			{
+				string[] VersionDirectories = Directory.GetDirectories(InDirectory);
+				foreach (string Dir in VersionDirectories)
+				{
+					string VersionString = Path.GetFileName(Dir);
+					Version? FoundVersion;
+					if (Version.TryParse(VersionString, out FoundVersion) && FoundVersion > LatestVersion)
+					{
+						if (NoLaterThan == null || FoundVersion <= NoLaterThan)
+						{
+							LatestVersion = FoundVersion;
+						}
+					}
+				}
+			}
+			return LatestVersion;
+		}
+
+		[SupportedOSPlatform("windows")]
+		private List<string> ExpandWinMDReferences(string SdkVersion, string[] WinMDReferences)
+		{
+			// Allow bringing in Windows SDK contracts just by naming the contract
+			// These are files that look like References/10.0.98765.0/AMadeUpWindowsApiContract/5.0.0.0/AMadeUpWindowsApiContract.winmd
+			List<string> ExpandedWinMDReferences = new List<string>();
+
+			VersionNumber? OutSdkVersion;
+			DirectoryReference? OutSdkDir;
+			if (!WindowsPlatform.TryGetWindowsSdkDir(SdkVersion, out OutSdkVersion, out OutSdkDir))
+			{
+				Log.TraceError("Failed to find WinSDK " + SdkVersion);
+				throw new Exception("Failed to find WinSDK");
+			}
+
+			// The first few releases of the Windows 10 SDK didn't put the SDK version in the reference path
+			string ReferenceRoot = Path.Combine(OutSdkDir.FullName, "References");
+			string VersionedReferenceRoot = Path.Combine(ReferenceRoot, OutSdkVersion.ToString());
+			if (Directory.Exists(VersionedReferenceRoot))
+			{
+				ReferenceRoot = VersionedReferenceRoot;
+			}
+
+			foreach (string WinMDRef in WinMDReferences)
+			{
+				if (File.Exists(WinMDRef))
+				{
+					// Already a valid path
+					ExpandedWinMDReferences.Add(WinMDRef);
+				}
+				else
+				{
+					string ContractFolder = Path.Combine(ReferenceRoot, WinMDRef);
+
+					Version ContractVersion = FindLatestVersionDirectory(ContractFolder, null);
+					string ExpandedWinMDRef = Path.Combine(ContractFolder, ContractVersion.ToString(), WinMDRef + ".winmd");
+					if (File.Exists(ExpandedWinMDRef))
+					{
+						ExpandedWinMDReferences.Add(ExpandedWinMDRef);
+					}
+					else
+					{
+						Log.TraceWarning("Unable to resolve location for HoloLens WinMD api contract {0}, file {1}", WinMDRef, ExpandedWinMDRef);
+					}
+				}
+			}
+
+			return ExpandedWinMDReferences;
+		}
+
+		/// <summary>
+		/// WinMD reference info
+		/// </summary>
+		/// <param name="InWindMDSourcePath"></param>
+		/// <param name="InPackageRelativeDllPath"></param>
+		/// <param name="SdkVersion"></param>
+		[SupportedOSPlatform("windows")]
+		public WinMDRegistrationInfo(FileReference InWindMDSourcePath, string InPackageRelativeDllPath, string SdkVersion)
+		{
+			PackageRelativeDllPath = InPackageRelativeDllPath;
+			ResolveSearchPaths.Add(InWindMDSourcePath.Directory.FullName);
+			List<string> WinMDAssemblies = ExpandWinMDReferences(SdkVersion, new string[] {
+				"Windows.Foundation.FoundationContract",
+				"Windows.Foundation.UniversalApiContract"
+			});
+			string[] runtimeAssemblies = Directory.GetFiles(RuntimeEnvironment.GetRuntimeDirectory(), "*.dll");
+			// Create the list of assembly paths consisting of runtime assemblies and the input file.
+			List<string> AssemblyPaths = new List<string>(runtimeAssemblies);
+			AssemblyPaths.AddRange(WinMDAssemblies);
+			AssemblyPaths.Add(InWindMDSourcePath.FullName);
+
+			// Create MetadataLoadContext that can resolve assemblies using the created list.
+			PathAssemblyResolver PathAssemblyResolver = new PathAssemblyResolver(AssemblyPaths);
+			MetadataLoadContext Mlc = new MetadataLoadContext(PathAssemblyResolver);
+
+			using (Mlc)
+			{
+				ActivatableTypesList = new List<ActivatableType>();
+				Assembly DependsOn = Mlc.LoadFromAssemblyPath(InWindMDSourcePath.FullName);
+				foreach (Type WinMDType in DependsOn.GetExportedTypes())
+				{
+					bool IsActivatable = false;
+					string ThreadingModelName = "both";
+					foreach (CustomAttributeData Attr in WinMDType.CustomAttributes)
+					{
+						if (Attr.AttributeType.FullName == "Windows.Foundation.Metadata.ActivatableAttribute" ||
+						    Attr.AttributeType.FullName == "Windows.Foundation.Metadata.StaticAttribute")
+						{
+							IsActivatable = true;
+						}
+						else if (Attr.AttributeType.FullName == "Windows.Foundation.Metadata.ThreadingAttribute")
+						{
+							CustomAttributeTypedArgument Argument = Attr.ConstructorArguments[0];
+							ThreadingModelName = Enum.GetName(Argument.ArgumentType, Argument.Value!)!.ToLowerInvariant();
+						}
+					}
+					if (IsActivatable)
+					{
+						ActivatableTypesList.Add(new ActivatableType(WinMDType.FullName!, ThreadingModelName));
+					}
+				}
+			}
+		}
+
+		/// <summary>
+		/// Path to the WinRT library
+		/// </summary>
+		public string PackageRelativeDllPath { get; private set; }
+
+		/// <summary>
+		/// List of the types in the WinMD
+		/// </summary>
+		public IEnumerable<ActivatableType> ActivatableTypes
+		{
+			get
+			{
+				return ActivatableTypesList;
+			}
+		}
+
+		private static List<string> ResolveSearchPaths = new List<string>();
+		private List<ActivatableType> ActivatableTypesList;
+	}
+}
diff --git a/Engine/Source/Programs/AutomationTool/HoloLens/HoloLens.Automation.csproj b/Engine/Source/Programs/AutomationTool/HoloLens/HoloLens.Automation.csproj
deleted file mode 100644
diff --git a/Engine/Source/Programs/AutomationTool/HoloLens/HoloLensPlatform.Automation.cs b/Engine/Source/Programs/AutomationTool/HoloLens/HoloLensPlatform.Automation.cs
deleted file mode 100644
diff --git a/Engine/Source/Programs/AutomationTool/HoloLens/Properties/AssemblyInfo.cs b/Engine/Source/Programs/AutomationTool/HoloLens/Properties/AssemblyInfo.cs
deleted file mode 100644
diff --git a/Engine/Source/Programs/UnrealBuildTool/Platform/HoloLens/HoloLensDeploy.cs b/Engine/Source/Programs/UnrealBuildTool/Platform/HoloLens/HoloLensDeploy.cs
deleted file mode 100644
diff --git a/Engine/Source/Programs/UnrealBuildTool/Platform/HoloLens/HoloLensExports.cs b/Engine/Source/Programs/UnrealBuildTool/Platform/HoloLens/HoloLensExports.cs
deleted file mode 100644
diff --git a/Engine/Source/Programs/UnrealBuildTool/Platform/HoloLens/HoloLensProjectGenerator.cs b/Engine/Source/Programs/UnrealBuildTool/Platform/HoloLens/HoloLensProjectGenerator.cs
deleted file mode 100644
diff --git a/Engine/Source/Programs/UnrealBuildTool/Platform/HoloLens/HoloLensToolChain.cs b/Engine/Source/Programs/UnrealBuildTool/Platform/HoloLens/HoloLensToolChain.cs
deleted file mode 100644
diff --git a/Engine/Source/Programs/UnrealBuildTool/Platform/HoloLens/ManifestGenerator.cs b/Engine/Source/Programs/UnrealBuildTool/Platform/HoloLens/ManifestGenerator.cs
deleted file mode 100644
diff --git a/Engine/Source/Programs/UnrealBuildTool/Platform/HoloLens/UEBuildHoloLens.cs b/Engine/Source/Programs/UnrealBuildTool/Platform/HoloLens/UEBuildHoloLens.cs
deleted file mode 100644
diff --git a/Engine/Source/Programs/UnrealBuildTool/Platform/HoloLens/WinMDRegistrationInfo.cs b/Engine/Source/Programs/UnrealBuildTool/Platform/HoloLens/WinMDRegistrationInfo.cs
deleted file mode 100644
