Add HoloLens UnrealGameSync.ini for platform extensions.

#rnx

diff --git a/Engine/Platforms/Hololens/Programs/UnrealGameSync/UnrealGameSync.ini b/Engine/Platforms/Hololens/Programs/UnrealGameSync/UnrealGameSync.ini
new file mode 100644
--- /dev/null
+++ b/Engine/Platforms/Hololens/Programs/UnrealGameSync/UnrealGameSync.ini
@@ -0,0 +1,3 @@
+[Options]
++SyncCategory=(UniqueId="c60c38c0-3ca8-46eb-b63d-c8f439d3d540", Name="Platform Support: HoloLens", Paths="/Engine/Platforms/HoloLens/...")
+
