Restore HoloLens native code analysis and architecture override.

#rnx

diff --git a/Engine/Platforms/Hololens/Source/Programs/UnrealBuildTool/HoloLensToolChain.cs b/Engine/Platforms/Hololens/Source/Programs/UnrealBuildTool/HoloLensToolChain.cs
--- a/Engine/Platforms/Hololens/Source/Programs/UnrealBuildTool/HoloLensToolChain.cs
+++ b/Engine/Platforms/Hololens/Source/Programs/UnrealBuildTool/HoloLensToolChain.cs
@@ -112,6 +112,39 @@ namespace UnrealBuildTool
 			Arguments.Add("/ignore:4264");
 		}
 
+		protected override void ModifyFinalCompileAction(VCCompileAction CompileAction, CppCompileEnvironment CompileEnvironment, FileItem SourceFile, DirectoryReference OutputDir, string ModuleName)
+		{
+			base.ModifyFinalCompileAction(CompileAction, CompileEnvironment, SourceFile, OutputDir, ModuleName);
+
+			if (Target.HoloLensPlatform.bRunNativeCodeAnalysis)
+			{
+				// Add the analysis log to the produced item list.
+				FileItem AnalysisLogFile = FileItem.GetItemByFileReference(
+					FileReference.Combine(
+						OutputDir,
+						Path.GetFileName(SourceFile.AbsolutePath) + ".nativecodeanalysis.xml"
+						)
+					); ;
+				CompileAction.AdditionalProducedItems.Add(AnalysisLogFile);
+				// Peform code analysis with results in a log file
+				CompileAction.Arguments.AddFormat("/analyze:log \"{0}\"", AnalysisLogFile.AbsolutePath);
+				// Suppress code analysis output
+				CompileAction.Arguments.Add("/analyze:quiet");
+				string? rulesetFile = Target.HoloLensPlatform.NativeCodeAnalysisRuleset;
+				if (!String.IsNullOrEmpty(rulesetFile))
+				{
+					if (!Path.IsPathRooted(rulesetFile))
+					{
+						rulesetFile = FileReference.Combine(Target.ProjectFile!.Directory, rulesetFile).FullName;
+					}
+					// A non default ruleset was specified
+					CompileAction.Arguments.AddFormat("/analyze:ruleset \"{0}\"", rulesetFile);
+				}
+			}
+		}
+
+
+
 		private static DirectoryReference? CurrentWindowsSdkBinDir = null;
 		private static Version? CurrentWindowsSdkVersion;
 
diff --git a/Engine/Platforms/Hololens/Source/Programs/UnrealBuildTool/UEBuildHoloLens.cs b/Engine/Platforms/Hololens/Source/Programs/UnrealBuildTool/UEBuildHoloLens.cs
--- a/Engine/Platforms/Hololens/Source/Programs/UnrealBuildTool/UEBuildHoloLens.cs
+++ b/Engine/Platforms/Hololens/Source/Programs/UnrealBuildTool/UEBuildHoloLens.cs
@@ -27,14 +27,6 @@ namespace UnrealBuildTool
 		[CommandLine("-2022", Value = nameof(WindowsCompiler.VisualStudio2022))]
 		public WindowsCompiler Compiler = WindowsCompiler.Default;
 
-		/// <summary>
-		/// Architecture of Target.
-		/// </summary>
-		public WindowsArchitecture Architecture
-		{
-			get;
-		}
-
 		/// <summary>
 		/// Enable PIX debugging (automatically disabled in Shipping and Test configs)
 		/// </summary>
@@ -71,26 +63,6 @@ namespace UnrealBuildTool
 		/// </summary>
 		[ConfigFile(ConfigHierarchyType.Engine, "/Script/HoloLensPlatformEditor.HoloLensTargetSettings", "NativeCodeAnalysisRuleset")]
 		public string? NativeCodeAnalysisRuleset = null;
-		
-		/// <summary>
-		/// Constructor
-		/// </summary>
-		/// <param name="Info">Target information</param>
-		public HoloLensTargetRules(TargetInfo Info)
-		{
-			if (Info.Platform == UnrealTargetPlatform.HoloLens && !String.IsNullOrEmpty(Info.Architecture))
-			{
-				Architecture = (WindowsArchitecture)Enum.Parse(typeof(WindowsArchitecture), Info.Architecture, true);
-			}
-		}
-
-		/// <summary>
-		/// Constructor - temporary #jira UE-149794
-		/// </summary>
-		public HoloLensTargetRules()
-		{
-			//#FIXME: need to set Architecture
-		}
 	}
 
 	/// <summary>
@@ -122,11 +94,6 @@ namespace UnrealBuildTool
 			get { return Inner.Compiler; }
 		}
 
-		public WindowsArchitecture Architecture
-		{
-			get { return Inner.Architecture; }
-		}
-
 		public bool bPixProfilingEnabled
 		{
 			get { return Inner.bPixProfilingEnabled; }
@@ -224,6 +191,12 @@ namespace UnrealBuildTool
 
 			ConfigCache.ReadSettings(IniDirRef, Platform, Target.HoloLensPlatform);
 
+			WindowsArchitecture Architecture = WindowsArchitecture.x64;
+			if( !String.IsNullOrEmpty(Target.Architecture))
+			{
+				Architecture = (WindowsArchitecture)Enum.Parse(typeof(WindowsArchitecture), Target.Architecture, true);
+			}
+
 			if (Target.HoloLensPlatform.Compiler == WindowsCompiler.Default)
 			{
 				if (CompilerBeforeReadSettings != WindowsCompiler.Default)
@@ -233,17 +206,17 @@ namespace UnrealBuildTool
 				}
 				else
 				{
-					Target.HoloLensPlatform.Compiler = WindowsPlatform.GetDefaultCompiler(Target.ProjectFile, Target.HoloLensPlatform.Architecture);
+					Target.HoloLensPlatform.Compiler = WindowsPlatform.GetDefaultCompiler(Target.ProjectFile, Architecture);
 				}
 			}
 
 			if (!Target.bGenerateProjectFiles)
 			{
-				Log.TraceInformationOnce("Using {0} architecture for deploying to HoloLens device", Target.HoloLensPlatform.Architecture);
+				Log.TraceInformationOnce("Using {0} architecture for deploying to HoloLens device", Architecture);
 			}
 
 			Target.WindowsPlatform.Compiler = Target.HoloLensPlatform.Compiler;
-			Target.WindowsPlatform.Architecture = Target.HoloLensPlatform.Architecture;
+			Target.WindowsPlatform.Architecture = Architecture;
 			Target.WindowsPlatform.bPixProfilingEnabled = Target.HoloLensPlatform.bPixProfilingEnabled;
 			Target.WindowsPlatform.bUseWindowsSDK10 = true;
 			Target.WindowsPlatform.bStrictConformanceMode = false;
diff --git a/Engine/Source/Programs/UnrealBuildTool/Platform/Windows/VCToolChain.cs b/Engine/Source/Programs/UnrealBuildTool/Platform/Windows/VCToolChain.cs
--- a/Engine/Source/Programs/UnrealBuildTool/Platform/Windows/VCToolChain.cs
+++ b/Engine/Source/Programs/UnrealBuildTool/Platform/Windows/VCToolChain.cs
@@ -1284,7 +1284,7 @@ namespace UnrealBuildTool
 					}
 
 					// Experimental: support for JSON output of timing data
-					if(Target.WindowsPlatform.Compiler.IsClang() && Target.WindowsPlatform.bClangTimeTrace)
+					if (Target.WindowsPlatform.Compiler.IsClang() && Target.WindowsPlatform.bClangTimeTrace)
 					{
 						CompileAction.Arguments.Add("-Xclang -ftime-trace");
 						CompileAction.AdditionalProducedItems.Add(FileItem.GetItemByFileReference(ObjectFile.Location.ChangeExtension(".json")));
@@ -1353,33 +1353,6 @@ namespace UnrealBuildTool
 					CompileAction.Arguments.Add(CompileEnvironment.AdditionalArguments);
 				}
 
-				/* FIXME: need to fix this #jira UE-149794
-				if (CompileEnvironment.Platform.IsInGroup(UnrealPlatformGroup.HoloLens) && Target.HoloLensPlatform.bRunNativeCodeAnalysis)
-				{
-					// Add the analysis log to the produced item list.
-					FileItem AnalysisLogFile = FileItem.GetItemByFileReference(
-						FileReference.Combine(
-							OutputDir,
-							Path.GetFileName(SourceFile.AbsolutePath) + ".nativecodeanalysis.xml"
-							)
-						); ;
-					CompileAction.AdditionalProducedItems.Add(AnalysisLogFile);
-					// Peform code analysis with results in a log file
-					CompileAction.Arguments.AddFormat("/analyze:log \"{0}\"", AnalysisLogFile.AbsolutePath);
-					// Suppress code analysis output
-					CompileAction.Arguments.Add("/analyze:quiet");
-					string? rulesetFile = Target.HoloLensPlatform.NativeCodeAnalysisRuleset;
-					if (!String.IsNullOrEmpty(rulesetFile))
-					{
-						if (!Path.IsPathRooted(rulesetFile))
-						{
-							rulesetFile = FileReference.Combine(Target.ProjectFile!.Directory, rulesetFile).FullName;
-						}
-						// A non default ruleset was specified
-						CompileAction.Arguments.AddFormat("/analyze:ruleset \"{0}\"", rulesetFile);
-					}
-				}*/
-
 				if (SourceFile.HasExtension(".ixx"))
 				{
 					FileItem IfcFile = FileItem.GetItemByFileReference(FileReference.Combine(GetModuleInterfaceDir(OutputDir), SourceFile.Location.ChangeExtension(".ifc").GetFileName()));
@@ -1436,6 +1409,9 @@ namespace UnrealBuildTool
 					}
 				}
 
+				// Allow derived toolchains to make further changes
+				ModifyFinalCompileAction(CompileAction, CompileEnvironment, SourceFile, OutputDir, ModuleName );
+
 				if (!ProjectFileGenerator.bGenerateProjectFiles)
 				{
 					CompileAction.WriteResponseFile(Graph);
@@ -1460,6 +1436,10 @@ namespace UnrealBuildTool
 			return Result;
 		}
 
+		protected virtual void ModifyFinalCompileAction(VCCompileAction CompileAction, CppCompileEnvironment CompileEnvironment, FileItem SourceFile, DirectoryReference OutputDir, string ModuleName)
+		{
+		}
+
 		private Action GenerateParseTimingInfoAction(FileItem SourceFile, FileItem TimingFile, IActionGraphBuilder Graph)
 		{
 			FileItem TimingJsonFile = FileItem.GetItemByPath(Path.ChangeExtension(TimingFile.AbsolutePath, ".cta"));
