move HoloLens runtime source to platform extension folder.

#rnx

diff --git a/Engine/Platforms/Hololens/Source/Programs/UnrealBuildTool/UEBuildHoloLens.cs b/Engine/Platforms/Hololens/Source/Programs/UnrealBuildTool/UEBuildHoloLens.cs
--- a/Engine/Platforms/Hololens/Source/Programs/UnrealBuildTool/UEBuildHoloLens.cs
+++ b/Engine/Platforms/Hololens/Source/Programs/UnrealBuildTool/UEBuildHoloLens.cs
@@ -493,8 +493,6 @@ namespace UnrealBuildTool
 			}
 			else if (ModuleName == "D3D12RHI")
 			{
-				// To enable platform specific D3D12 RHI Types
-				Rules.PrivateIncludePaths.Add("Runtime/D3D12RHI/Private/HoloLens");
 				Rules.PrivateDependencyModuleNames.Add("WinPixEventRuntime");
 			}
 			else if (ModuleName == "DX11")
@@ -602,6 +600,7 @@ namespace UnrealBuildTool
 			CompileEnvironment.Definitions.Add("HOLOLENS=1");
 
 			CompileEnvironment.Definitions.Add("WINAPI_FAMILY=WINAPI_FAMILY_APP");
+			CompileEnvironment.Definitions.Add("PLATFORM_IS_EXTENSION=1");
 			CompileEnvironment.Definitions.Add("PLATFORM_MICROSOFT=1");
 
 			// No D3DX on HoloLens!
diff --git a/Engine/Platforms/Hololens/Source/Runtime/ApplicationCore/ApplicationCore_HoloLens.build.cs b/Engine/Platforms/Hololens/Source/Runtime/ApplicationCore/ApplicationCore_HoloLens.build.cs
new file mode 100644
--- /dev/null
+++ b/Engine/Platforms/Hololens/Source/Runtime/ApplicationCore/ApplicationCore_HoloLens.build.cs
@@ -0,0 +1,11 @@
+// Copyright Epic Games, Inc. All Rights Reserved.
+
+namespace UnrealBuildTool.Rules
+{
+	public class ApplicationCore_HoloLens : ApplicationCore
+	{
+		public ApplicationCore_HoloLens(ReadOnlyTargetRules Target) : base(Target)
+		{
+		}
+	}
+}
diff --git a/Engine/Platforms/Hololens/Source/Runtime/ApplicationCore/Private/HoloLensApplication.cpp b/Engine/Platforms/Hololens/Source/Runtime/ApplicationCore/Private/HoloLensApplication.cpp
new file mode 100644
--- /dev/null
+++ b/Engine/Platforms/Hololens/Source/Runtime/ApplicationCore/Private/HoloLensApplication.cpp
@@ -0,0 +1,313 @@
+// Copyright Epic Games, Inc. All Rights Reserved.
+
+#include "HoloLensApplication.h"
+#include "HoloLensWindow.h"
+#include "HoloLensCursor.h"
+#include "HoloLensInputInterface.h"
+#include "HoloLensPlatformMisc.h"
+#include "GenericPlatform/GenericApplication.h"
+#include "Misc/CoreDelegates.h"
+#include "Misc/App.h"
+#include "Misc/ConfigCacheIni.h"
+#include "IInputDeviceModule.h"
+#include "IInputDevice.h"
+
+DEFINE_LOG_CATEGORY(LogHoloLens);
+
+static FHoloLensApplication* HoloLensApplication = NULL;
+FVector2D FHoloLensApplication::DesktopSize;
+
+#if PLATFORM_HOLOLENS
+Windows::Graphics::Holographic::HolographicSpace^ FHoloLensApplication::HoloSpace;
+#endif
+
+bool FHoloLensApplication::buildForRetailWindowsStore = false;
+
+FHoloLensApplication* FHoloLensApplication::CreateHoloLensApplication()
+{
+	GConfig->GetBool(TEXT("/Script/HoloLensPlatformEditor.HoloLensTargetSettings"), TEXT("bBuildForRetailWindowsStore"), buildForRetailWindowsStore, GEngineIni);
+	
+	check( HoloLensApplication == NULL );
+	HoloLensApplication = new FHoloLensApplication();
+
+	HoloLensApplication->InitLicensing();
+
+	return HoloLensApplication;
+}
+
+FHoloLensApplication* FHoloLensApplication::GetHoloLensApplication()
+{
+	return HoloLensApplication;
+}
+
+FHoloLensApplication::FHoloLensApplication()
+	: GenericApplication( MakeShareable( new FHoloLensCursor() ) )
+	, InputInterface( FHoloLensInputInterface::Create( MessageHandler ) )
+	// @MIXEDREALITY_CHANGE : BEGIN
+	, bHasLoadedInputPlugins(false)
+	// @MIXEDREALITY_CHANGE : END
+	, ApplicationWindow(FHoloLensWindow::Make())
+{
+}
+
+TSharedRef< FGenericWindow > FHoloLensApplication::MakeWindow()
+{
+	return ApplicationWindow;
+}
+
+void FHoloLensApplication::InitializeWindow(const TSharedRef< FGenericWindow >& InWindow, const TSharedRef< FGenericWindowDefinition >& InDefinition, const TSharedPtr< FGenericWindow >& InParent, const bool bShowImmediately)
+{
+	const TSharedRef< FHoloLensWindow > Window = StaticCastSharedRef< FHoloLensWindow >(InWindow);
+
+	Window->Initialize(this, InDefinition);
+}
+
+IInputInterface* FHoloLensApplication::GetInputInterface()
+{
+	// NOTE: This does not increase the reference count, so don't cache the result
+	return InputInterface.Get();
+}
+
+void FHoloLensApplication::PollGameDeviceState( const float TimeDelta )
+{
+	// @MIXEDREALITY_CHANGE : BEGIN Mixed Reality spatial input support
+	// Initialize any externally-implemented input devices (we delay load initialize the array so any plugins have had time to load)
+	if (!bHasLoadedInputPlugins && GIsRunning)
+	{
+		TArray< IInputDeviceModule* > PluginImplementations = IModularFeatures::Get().GetModularFeatureImplementations< IInputDeviceModule >(IInputDeviceModule::GetModularFeatureName());
+		for (auto InputPluginIt = PluginImplementations.CreateIterator(); InputPluginIt; ++InputPluginIt)
+		{
+			TSharedPtr< IInputDevice > Device = (*InputPluginIt)->CreateInputDevice(MessageHandler);
+			AddExternalInputDevice(Device);
+		}
+
+		bHasLoadedInputPlugins = true;
+	}
+
+	//TODO:
+	if (FApp::UseVRFocus() && !FApp::HasVRFocus())
+	{
+		return; // do not proceed if the app uses VR focus but doesn't have it
+	}
+
+	// Poll game device state and send new events
+	InputInterface->Tick(TimeDelta);
+	InputInterface->SendControllerEvents();
+
+	// Poll externally-implemented devices
+	for (auto DeviceIt = ExternalInputDevices.CreateIterator(); DeviceIt; ++DeviceIt)
+	{
+		(*DeviceIt)->Tick(TimeDelta);
+		(*DeviceIt)->SendControllerEvents();
+	}
+	// @MIXEDREALITY_CHANGE : END Mixed Reality spatial input support
+}
+
+FVector2D FHoloLensApplication::GetDesktopSize()
+{
+	return DesktopSize;
+}
+
+bool FHoloLensApplication::IsMouseAttached() const
+{
+	return (ref new Windows::Devices::Input::MouseCapabilities())->MousePresent > 0;
+}
+
+bool FHoloLensApplication::IsGamepadAttached() const
+{
+	//@MIXEDREALITY_CHANGE : BEGIN
+	bool gamepadAttached = false;
+	
+	gamepadAttached = Windows::Gaming::Input::Gamepad::Gamepads->Size > 0;
+	if (gamepadAttached)
+	{
+		return true;
+	}
+
+	for( auto DeviceIt = ExternalInputDevices.CreateConstIterator(); DeviceIt; ++DeviceIt )
+	{
+		if ((*DeviceIt)->IsGamepadAttached())
+		{
+			return true;
+		}
+	}
+	
+	return false;
+	//@MIXEDREALITY_CHANGE : END
+}
+
+FPlatformRect FHoloLensApplication::GetWorkArea(const FPlatformRect& CurrentWindow) const
+{
+	// No good way of accounting for desktop task bar here (if present) since
+	// it's not a Universal Windows construct.  Best we can do is return the
+	// full desktop size.
+
+	FPlatformRect WorkArea;
+	WorkArea.Left = 0;
+	WorkArea.Top = 0;
+	WorkArea.Right = DesktopSize.X;
+	WorkArea.Bottom = DesktopSize.Y;
+
+	return WorkArea;
+}
+
+void FHoloLensApplication::GetInitialDisplayMetrics(FDisplayMetrics& OutDisplayMetrics) const
+{
+	FDisplayMetrics::RebuildDisplayMetrics(OutDisplayMetrics);
+
+	// The initial virtual display rect is used to constrain maximum window size, and
+	// so should report the full desktop resolution.  Later calls will report window
+	// size in order to apply a tightly fitting hit test structure.
+	OutDisplayMetrics.VirtualDisplayRect = OutDisplayMetrics.PrimaryDisplayWorkAreaRect;
+}
+
+void FHoloLensApplication::CacheDesktopSize()
+{
+#if WIN10_SDK_VERSION >= 14393
+	if (Windows::Foundation::Metadata::ApiInformation::IsPropertyPresent("Windows.Graphics.Display.DisplayInformation", "ScreenHeightInRawPixels"))
+	{
+		Windows::Graphics::Display::DisplayInformation^ DisplayInfo = Windows::Graphics::Display::DisplayInformation::GetForCurrentView();
+		DesktopSize.X = DisplayInfo->ScreenWidthInRawPixels;
+		DesktopSize.Y = DisplayInfo->ScreenHeightInRawPixels;
+	}
+	else
+#endif
+	{
+		// Note this only works *before* the CoreWindow has been activated and received its first resize event.
+		Windows::UI::ViewManagement::ApplicationView^ ViewManagementView = Windows::UI::ViewManagement::ApplicationView::GetForCurrentView();
+		float Dpi = static_cast<uint32_t>(Windows::Graphics::Display::DisplayInformation::GetForCurrentView()->LogicalDpi);
+		FVector2D Size;
+		Size.X = FHoloLensWindow::ConvertDipsToPixels(ViewManagementView->VisibleBounds.Width, Dpi);
+		Size.Y = FHoloLensWindow::ConvertDipsToPixels(ViewManagementView->VisibleBounds.Height, Dpi);
+		DesktopSize = Size;
+	}
+}
+
+void FDisplayMetrics::RebuildDisplayMetrics(FDisplayMetrics& OutDisplayMetrics)
+{
+	FVector2D DesktopSize = FHoloLensApplication::GetDesktopSize();
+
+	OutDisplayMetrics.PrimaryDisplayWidth = DesktopSize.X;
+	OutDisplayMetrics.PrimaryDisplayHeight = DesktopSize.Y;
+
+	// As with FHoloLensApplication::GetWorkArea we can't really measure the difference between
+	// desktop size and work area on HoloLens.
+	OutDisplayMetrics.PrimaryDisplayWorkAreaRect.Left = 0;
+	OutDisplayMetrics.PrimaryDisplayWorkAreaRect.Top = 0;
+	OutDisplayMetrics.PrimaryDisplayWorkAreaRect.Right = DesktopSize.X;
+	OutDisplayMetrics.PrimaryDisplayWorkAreaRect.Bottom = DesktopSize.Y;
+
+	// We can't get a proper VirtualDisplayRect that's the equivalent of Windows Desktop.
+	// We can, however, take advantage of the fact that one of the primary purposes of this
+	// property is to define the coordinate space for mouse events.  Reporting the window
+	// size aligns this with the coordinate space naturally used for HoloLens pointer events,
+	// and allows us to correctly handle cases where a windowed app spans multiple screens.
+	OutDisplayMetrics.VirtualDisplayRect = FHoloLensWindow::GetOSWindowBounds();
+
+	// Translate so that the top left of the Window is the origin to exactly match
+	// pointer event coordinate space.
+	OutDisplayMetrics.VirtualDisplayRect.Bottom -= OutDisplayMetrics.VirtualDisplayRect.Top;
+	OutDisplayMetrics.VirtualDisplayRect.Right -= OutDisplayMetrics.VirtualDisplayRect.Left;
+	OutDisplayMetrics.VirtualDisplayRect.Left = 0;
+	OutDisplayMetrics.VirtualDisplayRect.Top = 0;
+
+	// Apply the debug safe zones
+	OutDisplayMetrics.ApplyDefaultSafeZones();
+}
+
+TSharedRef< class FGenericApplicationMessageHandler > FHoloLensApplication::GetMessageHandler() const
+{
+	return MessageHandler;
+}
+
+void FHoloLensApplication::SetMessageHandler( const TSharedRef< FGenericApplicationMessageHandler >& InMessageHandler )
+{
+	GenericApplication::SetMessageHandler(InMessageHandler);
+	InputInterface->SetMessageHandler( InMessageHandler );
+	
+	//@MIXEDREALITY_CHANGE : BEGIN
+	TArray<IInputDeviceModule*> PluginImplementations = IModularFeatures::Get().GetModularFeatureImplementations<IInputDeviceModule>( IInputDeviceModule::GetModularFeatureName() );
+	for( auto DeviceIt = ExternalInputDevices.CreateIterator(); DeviceIt; ++DeviceIt )
+	{
+		(*DeviceIt)->SetMessageHandler(InMessageHandler);
+	}
+	//@MIXEDREALITY_CHANGE : END
+}
+
+void FHoloLensApplication::PumpMessages(const float TimeDelta)
+{
+	//TODO: message pump here?  Keep track of active dispatcher?
+	FHoloLensMisc::PumpMessages(true);
+}
+
+TSharedRef< class FHoloLensCursor > FHoloLensApplication::GetCursor() const
+{
+	return StaticCastSharedPtr<FHoloLensCursor>( Cursor ).ToSharedRef();
+}
+
+bool FHoloLensApplication::ApplicationLicenseValid(FPlatformUserId PlatformUser)
+{
+	try
+	{
+		if (buildForRetailWindowsStore)
+		{
+			return Windows::ApplicationModel::Store::CurrentApp::LicenseInformation->IsActive;
+		}
+		else
+		{
+			return Windows::ApplicationModel::Store::CurrentAppSimulator::LicenseInformation->IsActive;
+		}
+	}
+	catch (...)
+	{
+		UE_LOG(LogCore, Warning, TEXT("Error retrieving license information.  This typically indicates that the incorrect version of CurrentStoreApp is being used.  The currently selected version is Windows::ApplicationModel::Store::%s.  Check USING_RETAIL_WINDOWS_STORE."), USING_RETAIL_WINDOWS_STORE != 0 ? TEXT("CurrentApp") : TEXT("CurrentAppSimulator"));
+		return false;
+	}
+}
+
+void FHoloLensApplication::InitLicensing()
+{
+	try
+	{
+		if (buildForRetailWindowsStore)
+		{
+			Windows::ApplicationModel::Store::CurrentApp::LicenseInformation->LicenseChanged += ref new Windows::ApplicationModel::Store::LicenseChangedEventHandler(LicenseChangedHandler);
+		}
+		else
+		{
+			Windows::ApplicationModel::Store::CurrentAppSimulator::LicenseInformation->LicenseChanged += ref new Windows::ApplicationModel::Store::LicenseChangedEventHandler(LicenseChangedHandler);
+		}
+	}
+	catch (...)
+	{
+		UE_LOG(LogCore, Warning, TEXT("Error registering for LicenseChanged event.  This typically indicates that the incorrect version of CurrentStoreApp is being used.  The currently selected version is Windows::ApplicationModel::Store::%s.  Check USING_RETAIL_WINDOWS_STORE."), USING_RETAIL_WINDOWS_STORE != 0 ? TEXT("CurrentApp") : TEXT("CurrentAppSimulator"));
+#if UE_BUILD_SHIPPING
+		// In shipping builds we treat this as fatal to reduce the chance of a build with misconfigured licensing being accidentally released.
+		throw;
+#endif
+	}
+}
+
+void FHoloLensApplication::LicenseChangedHandler()
+{
+	Windows::ApplicationModel::Core::CoreApplication::MainView->Dispatcher->RunAsync(
+		Windows::UI::Core::CoreDispatcherPriority::Normal,
+		ref new Windows::UI::Core::DispatchedHandler(BroadcastUELicenseChangeEvent));
+}
+
+void FHoloLensApplication::BroadcastUELicenseChangeEvent()
+{
+	FCoreDelegates::ApplicationLicenseChange.Broadcast();
+}
+
+// @MIXEDREALITY_CHANGE : BEGIN Mixed Reality spatial input support
+void FHoloLensApplication::AddExternalInputDevice(TSharedPtr<IInputDevice> InputDevice)
+{
+	if (InputDevice.IsValid())
+	{
+		ExternalInputDevices.Add(InputDevice);
+	}
+}
+// @MIXEDREALITY_CHANGE : END Mixed Reality spatial input support
+
+
diff --git a/Engine/Platforms/Hololens/Source/Runtime/ApplicationCore/Private/HoloLensCursor.cpp b/Engine/Platforms/Hololens/Source/Runtime/ApplicationCore/Private/HoloLensCursor.cpp
new file mode 100644
--- /dev/null
+++ b/Engine/Platforms/Hololens/Source/Runtime/ApplicationCore/Private/HoloLensCursor.cpp
@@ -0,0 +1,274 @@
+// Copyright Epic Games, Inc. All Rights Reserved.
+
+#include "HoloLensCursor.h"
+#include "HoloLensApplication.h"
+#include "HoloLensWindow.h"
+
+using namespace Windows::ApplicationModel::Core;
+using namespace Windows::UI::Core;
+
+FHoloLensCursorMouseEventObj::FHoloLensCursorMouseEventObj()
+{
+}
+
+Windows::Foundation::TypedEventHandler<Windows::Devices::Input::MouseDevice ^, Windows::Devices::Input::MouseEventArgs ^>^ FHoloLensCursorMouseEventObj::GetMouseMovedHandler()
+{
+    return ref new Windows::Foundation::TypedEventHandler<Windows::Devices::Input::MouseDevice ^, Windows::Devices::Input::MouseEventArgs ^>(this, &FHoloLensCursorMouseEventObj::OnMouseMoved);
+}
+
+void FHoloLensCursorMouseEventObj::OnMouseMoved(Windows::Devices::Input::MouseDevice ^sender, Windows::Devices::Input::MouseEventArgs ^args)
+{
+    FHoloLensApplication* Application = FHoloLensApplication::GetHoloLensApplication();
+    if (Application != NULL)
+    {
+		Application->GetCursor()->OnRawMouseMove(FIntVector(args->MouseDelta.X, args->MouseDelta.Y, 0));
+    }
+}
+
+FHoloLensCursor::FHoloLensCursor()
+{
+    Cursors = ref new Platform::Array<CoreCursor^>((int)EMouseCursor::TotalCursorCount);
+    MouseEventObj = ref new FHoloLensCursorMouseEventObj();
+    bUsingRawMouseNoCursor = false;
+    bDeferredCursorTypeChange = false;
+
+	// Load up cursors that we'll be using
+	for( int32 CurCursorIndex = 0; CurCursorIndex < EMouseCursor::TotalCursorCount; ++CurCursorIndex )
+	{
+		Cursors[ CurCursorIndex ] = GetDefaultCursorForType(static_cast<EMouseCursor::Type>(CurCursorIndex));
+	}
+}
+
+FHoloLensCursor::~FHoloLensCursor()
+{
+	// Release cursors
+	// NOTE: Shared cursors will automatically be destroyed when the application is destroyed.
+	//       For dynamically created cursors, use DestroyCursor
+	for( int32 CurCursorIndex = 0; CurCursorIndex < EMouseCursor::TotalCursorCount; ++CurCursorIndex )
+	{
+		switch( CurCursorIndex )
+		{
+		case EMouseCursor::None:
+		case EMouseCursor::Default:
+		case EMouseCursor::TextEditBeam:
+		case EMouseCursor::ResizeLeftRight:
+		case EMouseCursor::ResizeUpDown:
+		case EMouseCursor::ResizeSouthEast:
+		case EMouseCursor::ResizeSouthWest:
+		case EMouseCursor::CardinalCross:
+		case EMouseCursor::Crosshairs:
+		case EMouseCursor::Hand:
+		case EMouseCursor::GrabHand:
+		case EMouseCursor::GrabHandClosed:
+		case EMouseCursor::SlashedCircle:
+		case EMouseCursor::EyeDropper:
+		case EMouseCursor::Custom:
+			// Standard shared cursors don't need to be destroyed
+			break;
+
+		default:
+			// Unrecognized cursor type!
+			check( 0 );
+			break;
+		}
+	}
+}
+
+
+FVector2D FHoloLensCursor::GetPosition() const
+{
+	return CursorPosition;
+}
+
+void FHoloLensCursor::SetPosition( const int32 X, const int32 Y )
+{
+    CursorPosition.X = X;
+    CursorPosition.Y = Y;
+}
+
+void FHoloLensCursor::ProcessDeferredActions()
+{
+    if (bDeferredCursorTypeChange)
+    {
+		SetType(CurrentCursor);
+    }
+
+	if (DeferredMoveEvents.Num() > 0)
+	{
+		FHoloLensApplication* Application = FHoloLensApplication::GetHoloLensApplication();
+		if (Application != nullptr)
+		{
+			float Dpi = static_cast<uint32_t>(Windows::Graphics::Display::DisplayInformation::GetForCurrentView()->LogicalDpi);
+			for (const FIntVector& MouseDelta : DeferredMoveEvents)
+			{
+				Application->GetMessageHandler()->OnRawMouseMove(FHoloLensWindow::ConvertDipsToPixels(MouseDelta.X, Dpi), FHoloLensWindow::ConvertDipsToPixels(MouseDelta.Y, Dpi));
+			}
+			Application->GetMessageHandler()->OnCursorSet();
+		}
+		DeferredMoveEvents.Empty();
+	}
+}
+
+void FHoloLensCursor::OnRawMouseMove(const FIntVector& MouseDelta)
+{
+	DeferredMoveEvents.Add(MouseDelta);
+}
+
+void FHoloLensCursor::SetType( const EMouseCursor::Type InNewCursor )
+{
+    checkf(InNewCursor < EMouseCursor::TotalCursorCount, TEXT("Invalid cursor(%d) supplied"), InNewCursor);
+
+	if (CurrentCursor != InNewCursor || bDeferredCursorTypeChange)
+	{
+		CurrentCursor = InNewCursor;
+		// if we're on the UI thread, change the cursor, otherwise queue a deferred change
+        CoreWindow^ window = CoreWindow::GetForCurrentThread();
+		if (nullptr == window)
+		{
+            bDeferredCursorTypeChange = true;
+		}
+		else
+		{
+			bDeferredCursorTypeChange = false;
+			SetUseRawMouse(CurrentCursor == EMouseCursor::None);
+
+			try
+			{
+				window->PointerCursor = Cursors[CurrentCursor];
+			}
+			catch (Platform::Exception^ Ex)
+			{
+				if (CurrentCursor == EMouseCursor::Custom && Cursors[CurrentCursor] != nullptr)
+				{
+					UE_LOG(LogCore, Error, TEXT("Failed to set custom cursor with resource id %u.  Either id was invalid or cursor was not compiled in.  OS error message : %s"), Cursors[CurrentCursor]->Id, Ex->Message->Data());
+				}
+				else
+				{
+					UE_LOG(LogCore, Error, TEXT("Failed to set cursor to type %d : %s"), static_cast<int32>(CurrentCursor), Ex->Message->Data());
+				}
+			}
+		}
+    }
+}
+
+void FHoloLensCursor::GetSize( int32& Width, int32& Height ) const
+{
+	Width = 16;
+	Height = 16;
+}
+
+void FHoloLensCursor::UpdatePosition( const FVector2D& NewPosition )
+{
+	CursorPosition = NewPosition;
+}
+
+void FHoloLensCursor::Show( bool bShow )
+{
+}
+
+void FHoloLensCursor::Lock( const RECT* const Bounds )
+{
+}
+
+void FHoloLensCursor::SetUseRawMouse(bool bUse)
+{
+	if (bUsingRawMouseNoCursor != bUse)
+	{
+		bUsingRawMouseNoCursor = bUse;
+
+		try
+		{
+			Windows::Devices::Input::MouseDevice^ LowLevelMouse = Windows::Devices::Input::MouseDevice::GetForCurrentView();
+			if (bUsingRawMouseNoCursor)
+			{
+				check(MouseEventRegistrationToken.Value == 0);
+				MouseEventRegistrationToken = (LowLevelMouse->MouseMoved += MouseEventObj->GetMouseMovedHandler());
+			}
+			else
+			{
+				check(MouseEventRegistrationToken.Value != 0);
+				LowLevelMouse->MouseMoved -= MouseEventRegistrationToken;
+				MouseEventRegistrationToken.Value = 0;
+			}
+		}
+		catch (Platform::Exception^ Ex)
+		{
+			UE_LOG(LogCore, Warning, TEXT("Exception managing registration for low-level mouse events: %s"), Ex->Message->Data());
+		}
+	}
+}
+
+void FHoloLensCursor::SetTypeShape(EMouseCursor::Type InCursorType, void* CursorHandle)
+{
+	if (CursorHandle != nullptr)
+	{
+		// This will succeed even if CursorResourceId is invalid.  The point of failure if
+		// someone supplied a bad value will be when we actually try to set the window cursor.
+		Cursors[InCursorType] = ref new CoreCursor(CoreCursorType::Custom, static_cast<uint32>(reinterpret_cast<uint64>(CursorHandle)));
+	}
+	else
+	{
+		Cursors[InCursorType] = GetDefaultCursorForType(InCursorType);
+	}
+
+	if (CurrentCursor == InCursorType)
+	{
+		SetType(InCursorType);
+	}
+}
+
+Windows::UI::Core::CoreCursor^ FHoloLensCursor::GetDefaultCursorForType(EMouseCursor::Type InCursorType)
+{
+	switch (InCursorType)
+	{
+	case EMouseCursor::None:
+	case EMouseCursor::Custom:
+		return nullptr;
+
+	case EMouseCursor::Default:
+		return ref new CoreCursor(CoreCursorType::Arrow, 0);
+
+	case EMouseCursor::TextEditBeam:
+		return ref new CoreCursor(CoreCursorType::IBeam, 0);
+
+	case EMouseCursor::ResizeLeftRight:
+		return ref new CoreCursor(CoreCursorType::SizeWestEast, 0);
+
+	case EMouseCursor::ResizeUpDown:
+		return ref new CoreCursor(CoreCursorType::SizeNorthSouth, 0);
+
+	case EMouseCursor::ResizeSouthEast:
+		return ref new CoreCursor(CoreCursorType::SizeNorthwestSoutheast, 0);
+
+	case EMouseCursor::ResizeSouthWest:
+		return ref new CoreCursor(CoreCursorType::SizeNortheastSouthwest, 0);
+
+	case EMouseCursor::CardinalCross:
+		return ref new CoreCursor(CoreCursorType::SizeAll, 0);
+
+	case EMouseCursor::Crosshairs:
+		return ref new CoreCursor(CoreCursorType::Cross, 0);
+
+	case EMouseCursor::Hand:
+		return ref new CoreCursor(CoreCursorType::Hand, 0);
+
+	case EMouseCursor::GrabHand:
+		return ref new CoreCursor(CoreCursorType::Hand, 0);
+
+	case EMouseCursor::GrabHandClosed:
+		return ref new CoreCursor(CoreCursorType::Hand, 0);
+
+	case EMouseCursor::SlashedCircle:
+		return ref new CoreCursor(CoreCursorType::UniversalNo, 0);
+
+	case EMouseCursor::EyeDropper:
+		return ref new CoreCursor(CoreCursorType::Arrow, 0);
+
+	default:
+		// Unrecognized cursor type!
+		check(0);
+		return nullptr;
+	}
+
+}
+
diff --git a/Engine/Platforms/Hololens/Source/Runtime/ApplicationCore/Private/HoloLensErrorOutputDevice.cpp b/Engine/Platforms/Hololens/Source/Runtime/ApplicationCore/Private/HoloLensErrorOutputDevice.cpp
new file mode 100644
--- /dev/null
+++ b/Engine/Platforms/Hololens/Source/Runtime/ApplicationCore/Private/HoloLensErrorOutputDevice.cpp
@@ -0,0 +1,92 @@
+// Copyright Epic Games, Inc. All Rights Reserved.
+
+#include "HoloLensErrorOutputDevice.h"
+
+#include "HoloLensApplication.h"
+#include "HoloLensPlatformApplicationMisc.h"
+#include "HAL/PlatformMisc.h"
+#include "Microsoft/WindowsHWrapper.h"
+#include "CoreGlobals.h"
+#include "Misc/CoreDelegates.h"
+#include "Misc/OutputDeviceRedirector.h"
+
+
+FHoloLensErrorOutputDevice::FHoloLensErrorOutputDevice()
+{}
+
+void FHoloLensErrorOutputDevice::Serialize(const TCHAR* Msg, ELogVerbosity::Type Verbosity, const class FName& Category)
+{
+	UE_DEBUG_BREAK();
+
+	if (!GIsCriticalError)
+	{
+		int32 LastError = ::GetLastError();
+
+		GIsCriticalError = 1;
+		TCHAR ErrorBuffer[1024];
+		ErrorBuffer[0] = 0;
+
+		if (LastError == 0)
+		{
+			UE_LOG(LogHoloLens, Log, TEXT("HoloLens GetLastError: %s (%i)"), FPlatformMisc::GetSystemErrorMessage(ErrorBuffer, 1024, LastError), LastError);
+		}
+		else
+		{
+			UE_LOG(LogHoloLens, Error, TEXT("HoloLens GetLastError: %s (%i)"), FPlatformMisc::GetSystemErrorMessage(ErrorBuffer, 1024, LastError), LastError);
+		}
+	}
+	else
+	{
+		UE_LOG(LogHoloLens, Error, TEXT("Error reentered: %s"), Msg);
+	}
+
+	if (GIsGuarded)
+	{
+		// Propagate error so structured exception handler can perform necessary work.
+#if PLATFORM_EXCEPTIONS_DISABLED
+		UE_DEBUG_BREAK();
+#endif
+		FPlatformMisc::RaiseException(1);
+	}
+	else
+	{
+		// We crashed outside the guarded code (e.g. appExit).
+		HandleError();
+		FPlatformMisc::RequestExit(true);
+	}
+}
+
+void FHoloLensErrorOutputDevice::HandleError()
+{
+	// make sure we don't report errors twice
+	static int32 CallCount = 0;
+	int32 NewCallCount = FPlatformAtomics::InterlockedIncrement(&CallCount);
+	if (NewCallCount != 1)
+	{
+		UE_LOG(LogHoloLens, Error, TEXT("HandleError re-entered."));
+		return;
+	}
+
+	GIsGuarded = 0;
+	GIsRunning = 0;
+	GIsCriticalError = 1;
+	GLogConsole = NULL;
+	GErrorHist[UE_ARRAY_COUNT(GErrorHist) - 1] = 0;
+
+	// Trigger the OnSystemFailure hook if it exists
+	// make sure it happens after GIsGuarded is set to 0 in case this hook crashes
+	FCoreDelegates::OnHandleSystemError.Broadcast();
+
+	// Dump the error and flush the log.
+#if !NO_LOGGING
+	FDebug::LogFormattedMessageWithCallstack(LogHoloLens.GetCategoryName(), __FILE__, __LINE__, TEXT("=== Critical error: ==="), GErrorHist, ELogVerbosity::Error);
+#endif
+	GLog->Panic();
+
+	// Copy to clipboard in non-cooked editor builds.
+	FPlatformApplicationMisc::ClipboardCopy(GErrorHist);
+
+	FPlatformMisc::SubmitErrorReport(GErrorHist, EErrorReportMode::Interactive);
+
+	FCoreDelegates::OnShutdownAfterError.Broadcast();
+}
\ No newline at end of file
diff --git a/Engine/Platforms/Hololens/Source/Runtime/ApplicationCore/Private/HoloLensInputInterface.cpp b/Engine/Platforms/Hololens/Source/Runtime/ApplicationCore/Private/HoloLensInputInterface.cpp
new file mode 100644
--- /dev/null
+++ b/Engine/Platforms/Hololens/Source/Runtime/ApplicationCore/Private/HoloLensInputInterface.cpp
@@ -0,0 +1,499 @@
+// Copyright Epic Games, Inc. All Rights Reserved.
+
+#include "HoloLensInputInterface.h"
+#include "Misc/CoreDelegates.h"
+#include "GenericPlatform/GenericApplication.h"
+#include "WindowsGamingInputInterface.h"
+
+#include "Microsoft/AllowMicrosoftPlatformTypes.h"
+#include <Winuser.h>
+#include "Microsoft/HideMicrosoftPlatformTypes.h"
+
+
+const uint32 HoloLensKeyCount = 256;
+bool HoloLensKeyState[HoloLensKeyCount];
+bool HoloLensKeyLastState[HoloLensKeyCount];
+bool HoloLensCharState[HoloLensKeyCount];
+bool HoloLensCharLastState[HoloLensKeyCount];
+
+void appHoloLensKeyEvent(uint32 InCode, bool bInWasPressed)
+{
+	if (InCode < HoloLensKeyCount)
+	{
+		HoloLensKeyState[InCode] = bInWasPressed;
+	}
+}
+
+void appHoloLensCharEvent(uint32 InCode, bool bInWasPressed)
+{
+	if (InCode < HoloLensKeyCount)
+	{
+		HoloLensCharState[InCode] = bInWasPressed;
+	}
+}
+
+TSharedRef< FHoloLensInputInterface > FHoloLensInputInterface::Create( const TSharedRef< FGenericApplicationMessageHandler >& InMessageHandler )
+{
+	return MakeShareable( new FHoloLensInputInterface( InMessageHandler ) );
+}
+
+FHoloLensInputInterface::FHoloLensInputInterface( const TSharedRef< FGenericApplicationMessageHandler >& InMessageHandler )
+	: MessageHandler( InMessageHandler )
+	, GamingInput( WindowsGamingInputInterface::Create(InMessageHandler) )
+{
+	FMemory::Memzero(HoloLensKeyState, HoloLensKeyCount);
+	FMemory::Memzero(HoloLensKeyLastState, HoloLensKeyCount);
+	FMemory::Memzero(HoloLensCharState, HoloLensKeyCount);
+	FMemory::Memzero(HoloLensCharLastState, HoloLensKeyCount);
+
+	// register focus handling events
+	FCoreDelegates::ApplicationHasReactivatedDelegate.AddRaw( this, &FHoloLensInputInterface::OnFocusGain );
+	FCoreDelegates::ApplicationWillDeactivateDelegate.AddRaw( this, &FHoloLensInputInterface::OnFocusLost );
+}
+
+FHoloLensInputInterface::~FHoloLensInputInterface()
+{
+	// unregister focus handling events
+	FCoreDelegates::ApplicationHasReactivatedDelegate.RemoveAll( this );
+	FCoreDelegates::ApplicationWillDeactivateDelegate.RemoveAll( this );
+}
+
+void FHoloLensInputInterface::SetMessageHandler( const TSharedRef< FGenericApplicationMessageHandler >& InMessageHandler )
+{
+	MessageHandler = InMessageHandler;
+
+	GamingInput->SetMessageHandler(InMessageHandler);
+}
+
+void FHoloLensInputInterface::Tick( float DeltaTime )
+{
+	ConditionalScanForKeyboardChanges(DeltaTime);
+}
+
+void FHoloLensInputInterface::SendControllerEvents()
+{
+	GamingInput->UpdateGamepads();
+}
+
+void FHoloLensInputInterface::SetForceFeedbackChannelValue(int32 ControllerId, FForceFeedbackChannelType ChannelType, float Value)
+{
+	FForceFeedbackValues Values;
+	GamingInput->GetVibration(ControllerId, Values.LeftLarge, Values.RightLarge, Values.LeftSmall, Values.RightSmall);
+
+	switch (ChannelType)
+	{
+	case FForceFeedbackChannelType::LEFT_LARGE:
+		Values.LeftLarge = Value;
+		break;
+
+	case FForceFeedbackChannelType::LEFT_SMALL:
+		Values.LeftSmall = Value;
+		break;
+
+	case FForceFeedbackChannelType::RIGHT_LARGE:
+		Values.RightLarge = Value;
+		break;
+
+	case FForceFeedbackChannelType::RIGHT_SMALL:
+		Values.RightSmall = Value;
+		break;
+
+	default:
+		// Unknown channel, so ignore it
+		break;
+	}
+
+	SetForceFeedbackChannelValues(ControllerId, Values);
+}
+
+void FHoloLensInputInterface::SetForceFeedbackChannelValues(int32 ControllerId, const FForceFeedbackValues &Values)
+{
+	GamingInput->SetVibration(ControllerId, Values.LeftLarge, Values.RightLarge, Values.LeftSmall, Values.RightSmall);
+}
+
+//Windows::Gaming::Input::Gamepad^ FHoloLensInputInterface::GetGamepadForUser(const int UserId)
+//{
+//	return GamingInput->GetGamepadFor(UserId);
+//}
+
+int FHoloLensInputInterface::GetUserIdForController(Windows::Gaming::Input::Gamepad^ Controller)
+{
+	// TODO: handle multi-user scenario in HoloLens
+	return 0;
+}
+
+/** Platform-specific backdoor to get the FPlatformUserId from an XboxUserId */
+//FPlatformUserId FHoloLensInputInterface::GetPlatformUserIdFromXboxUserId(const TCHAR* XboxUserId)
+//{
+//	// TODO: handle multi-user scenario in HoloLens
+//	return 0;
+//}
+
+
+Windows::Gaming::Input::Gamepad^ FHoloLensInputInterface::GetGamepadForControllerId(int32 ControllerId)
+{
+	return GamingInput->GetGamepadForIndex(ControllerId);
+}
+
+
+void FHoloLensInputInterface::OnFocusLost()
+{
+	FMemory::Memzero(HoloLensKeyState);
+	FMemory::Memzero(HoloLensCharState);
+
+	GamingInput->DisableUpdate();
+	ConditionalScanForKeyboardChanges(0.0f);
+}
+
+void FHoloLensInputInterface::OnFocusGain()
+{
+	GamingInput->EnableUpdate();
+}
+
+uint32 LegacyMapVirtualKey(uint32 InVirtualKey)
+{
+	if ((InVirtualKey >= 48 && InVirtualKey <= 57)
+		|| (InVirtualKey >= 65 && InVirtualKey <= 90)
+		|| InVirtualKey == 8
+		|| InVirtualKey == 9
+		|| InVirtualKey == 13
+		|| InVirtualKey == 27
+		|| InVirtualKey == 32)
+	{
+		return InVirtualKey;
+	}
+
+	if (InVirtualKey >= 96 && InVirtualKey <= 105)
+	{
+		return InVirtualKey - 48;
+	}
+
+	if ((InVirtualKey == 106 || InVirtualKey == 107)
+		|| (InVirtualKey >= 109 && InVirtualKey <= 111))
+	{
+		return InVirtualKey - 64;
+	}
+
+	// Special keys
+	switch (InVirtualKey)
+	{
+	case 0x6A: return '*';
+	case 0x6B: return '+';
+	case 0x6D: return '-';
+	case 0x6E: return '.';
+	case 0x6F: return '/';
+
+	case 0xBA: return ';';  // VK_OEM_1
+	case 0xBB: return '=';  // VK_OEM_PLUS
+	case 0xBC: return ',';  // VK_OEM_COMMA
+	case 0xBD: return '_';  // VK_OEM_MINUS
+	case 0xBE: return '.';  // VK_OEM_PERIOD
+	case 0xBF: return '/';  // VK_OEM_2
+	case 0xC0: return '`';  // VK_OEM_3
+
+	case 0xDB: return '[';  // VK_OEM_4
+	case 0xDC: return '\\'; // VK_OEM_5
+	case 0xDD: return ']';  // VK_OEM_6
+	case 0xDE: return '\''; // VK_OEM_7
+	}
+
+	return 0;
+}
+
+void FHoloLensInputInterface::ConditionalScanForKeyboardChanges( float DeltaTime )
+{
+	bool bShiftIsDown = (HoloLensKeyState[VK_LSHIFT] == true) | (HoloLensKeyState[VK_RSHIFT] == true);
+	FModifierKeysState ModifierKeyState(
+		(HoloLensKeyState[VK_LSHIFT] == true) ? true : false,
+		(HoloLensKeyState[VK_RSHIFT] == true) ? true : false,
+		(HoloLensKeyState[VK_LCONTROL] == true) ? true : false,
+		(HoloLensKeyState[VK_RCONTROL] == true) ? true : false,
+		false,
+		false,
+		false,
+		false,
+		false
+		);
+
+	for (uint32 KeyIndex = 0; KeyIndex < HoloLensKeyCount; KeyIndex++)
+	{
+		// Process key up/down
+		{
+			bool bKeyPressed = HoloLensKeyState[KeyIndex];
+			bool bKeyReleased = ((HoloLensKeyState[KeyIndex] == false) && (HoloLensKeyLastState[KeyIndex] == true));
+
+			if ((bKeyPressed == true) || (bKeyReleased == true))
+			{
+				bool bIsRepeat = ((bKeyPressed == true) && (HoloLensKeyLastState[KeyIndex] == true));
+				
+				// Get the character code from the virtual key pressed.  If 0, no translation from virtual key to character exists
+				uint32 CharCode = LegacyMapVirtualKey(KeyIndex);
+
+				[REDACTED_EMAIL]: Put in a delay for registering repeats??
+				if (bIsRepeat == false)
+				{
+					if (bKeyReleased == false)
+					{
+						MessageHandler->OnKeyDown( KeyIndex, CharCode, bIsRepeat );
+					}
+					else
+					{
+						MessageHandler->OnKeyUp( KeyIndex, CharCode, bIsRepeat );
+					}
+				}
+
+				if (bKeyReleased == true)
+				{
+					uint32 CharKeyIndex = MapVirtualKeyToCharacter(KeyIndex, bShiftIsDown);
+					// Mark the char as released...
+					if (CharKeyIndex != 0)
+					{
+						HoloLensCharState[CharKeyIndex] = false;
+						HoloLensCharLastState[CharKeyIndex] = false;
+					}
+				}
+			}
+		}
+
+		// Process char
+		bool bCharPressed = HoloLensCharState[KeyIndex];
+		if (bCharPressed == true)
+		{
+			bool bIsRepeat = HoloLensCharLastState[KeyIndex];
+			[REDACTED_EMAIL]: Put in a delay for registering repeats??
+			if (bIsRepeat == false)
+			{
+				MessageHandler->OnKeyChar( KeyIndex, bIsRepeat );
+			}
+		}
+		
+		HoloLensKeyLastState[KeyIndex] = HoloLensKeyState[KeyIndex];
+		HoloLensCharLastState[KeyIndex] = HoloLensCharState[KeyIndex];
+	}
+}
+
+uint32 FHoloLensInputInterface::MapVirtualKeyToCharacter(uint32 InVirtualKey, bool bShiftIsDown)
+{
+	uint32 TranslatedKeyCode = (uint32)InVirtualKey;
+
+	switch (InVirtualKey)
+	{
+	case VK_LBUTTON:
+	case VK_RBUTTON:
+	case VK_MBUTTON:
+	case VK_XBUTTON1:
+	case VK_XBUTTON2:
+	case VK_BACK:
+	case VK_TAB:
+	case VK_RETURN:
+	case VK_PAUSE:
+	case VK_CAPITAL:
+	case VK_ESCAPE:
+	case VK_END:
+	case VK_HOME:
+	case VK_LEFT:
+	case VK_UP:
+	case VK_RIGHT:
+	case VK_DOWN:
+	case VK_INSERT:
+	case VK_DELETE:
+		TranslatedKeyCode = 0;
+		break;
+
+	case VK_SPACE:
+		TranslatedKeyCode = TEXT(' ');
+		break;
+
+	// VK_0 - VK_9 are the same as ASCII '0' - '9' (0x30 - 0x39)
+	// 0x40 : unassigned
+	// VK_A - VK_Z are the same as ASCII 'A' - 'Z' (0x41 - 0x5A)
+	case 0x41:
+		TranslatedKeyCode = (!bShiftIsDown) ? TEXT('a') : TEXT('A');
+		break;
+	case 0x42:
+		TranslatedKeyCode = (!bShiftIsDown) ? TEXT('b') : TEXT('B');
+		break;
+	case 0x43:
+		TranslatedKeyCode = (!bShiftIsDown) ? TEXT('c') : TEXT('C');
+		break;
+	case 0x44:
+		TranslatedKeyCode = (!bShiftIsDown) ? TEXT('d') : TEXT('D');
+		break;
+	case 0x45:
+		TranslatedKeyCode = (!bShiftIsDown) ? TEXT('e') : TEXT('E');
+		break;
+	case 0x46:
+		TranslatedKeyCode = (!bShiftIsDown) ? TEXT('f') : TEXT('F');
+		break;
+	case 0x47:
+		TranslatedKeyCode = (!bShiftIsDown) ? TEXT('g') : TEXT('G');
+		break;
+	case 0x48:
+		TranslatedKeyCode = (!bShiftIsDown) ? TEXT('h') : TEXT('H');
+		break;
+	case 0x49:
+		TranslatedKeyCode = (!bShiftIsDown) ? TEXT('i') : TEXT('I');
+		break;
+	case 0x4A:
+		TranslatedKeyCode = (!bShiftIsDown) ? TEXT('j') : TEXT('J');
+		break;
+	case 0x4B:
+		TranslatedKeyCode = (!bShiftIsDown) ? TEXT('k') : TEXT('K');
+		break;
+	case 0x4C:
+		TranslatedKeyCode = (!bShiftIsDown) ? TEXT('l') : TEXT('L');
+		break;
+	case 0x4D:
+		TranslatedKeyCode = (!bShiftIsDown) ? TEXT('m') : TEXT('M');
+		break;
+	case 0x4E:
+		TranslatedKeyCode = (!bShiftIsDown) ? TEXT('n') : TEXT('N');
+		break;
+	case 0x4F:
+		TranslatedKeyCode = (!bShiftIsDown) ? TEXT('o') : TEXT('O');
+		break;
+	case 0x50:
+		TranslatedKeyCode = (!bShiftIsDown) ? TEXT('p') : TEXT('P');
+		break;
+	case 0x51:
+		TranslatedKeyCode = (!bShiftIsDown) ? TEXT('q') : TEXT('Q');
+		break;
+	case 0x52:
+		TranslatedKeyCode = (!bShiftIsDown) ? TEXT('r') : TEXT('R');
+		break;
+	case 0x53:
+		TranslatedKeyCode = (!bShiftIsDown) ? TEXT('s') : TEXT('S');
+		break;
+	case 0x54:
+		TranslatedKeyCode = (!bShiftIsDown) ? TEXT('t') : TEXT('T');
+		break;
+	case 0x55:
+		TranslatedKeyCode = (!bShiftIsDown) ? TEXT('u') : TEXT('U');
+		break;
+	case 0x56:
+		TranslatedKeyCode = (!bShiftIsDown) ? TEXT('v') : TEXT('V');
+		break;
+	case 0x57:
+		TranslatedKeyCode = (!bShiftIsDown) ? TEXT('w') : TEXT('W');
+		break;
+	case 0x58:
+		TranslatedKeyCode = (!bShiftIsDown) ? TEXT('x') : TEXT('X');
+		break;
+	case 0x59:
+		TranslatedKeyCode = (!bShiftIsDown) ? TEXT('y') : TEXT('Y');
+		break;
+	case 0x5A:
+		TranslatedKeyCode = (!bShiftIsDown) ? TEXT('z') : TEXT('Z');
+		break;
+
+	case 0x30:
+	case VK_NUMPAD0:
+		TranslatedKeyCode = (!bShiftIsDown) ? TEXT('0') : TEXT(')'); 
+		break;
+	case 0x31:
+	case VK_NUMPAD1:
+		TranslatedKeyCode = (!bShiftIsDown) ? TEXT('1') : TEXT('!'); 
+		break;
+	case 0x32:
+	case VK_NUMPAD2:
+		TranslatedKeyCode = (!bShiftIsDown) ? TEXT('2') : TEXT('@'); 
+		break;
+	case 0x33:
+	case VK_NUMPAD3:
+		TranslatedKeyCode = (!bShiftIsDown) ? TEXT('3') : TEXT('#'); 
+		break;
+	case 0x34:
+	case VK_NUMPAD4:
+		TranslatedKeyCode = (!bShiftIsDown) ? TEXT('4') : TEXT('$'); 
+		break;
+	case 0x35:
+	case VK_NUMPAD5:
+		TranslatedKeyCode = (!bShiftIsDown) ? TEXT('5') : TEXT('%'); 
+		break;
+	case 0x36:
+	case VK_NUMPAD6:
+		TranslatedKeyCode = (!bShiftIsDown) ? TEXT('6') : TEXT('^'); 
+		break;
+	case 0x37:
+	case VK_NUMPAD7:
+		TranslatedKeyCode = (!bShiftIsDown) ? TEXT('7') : TEXT('&'); 
+		break;
+	case 0x38:
+	case VK_NUMPAD8:
+		TranslatedKeyCode = (!bShiftIsDown) ? TEXT('8') : TEXT('*'); 
+		break;
+	case 0x39:
+	case VK_NUMPAD9:
+		TranslatedKeyCode = (!bShiftIsDown) ? TEXT('9') : TEXT('('); 
+		break;
+
+	case VK_MULTIPLY:	
+		TranslatedKeyCode = TEXT('*'); 
+		break;
+	case VK_ADD:		
+		TranslatedKeyCode = TEXT('+'); 
+		break;
+	case VK_SUBTRACT:	
+		TranslatedKeyCode = TEXT('-'); 
+		break;
+	case VK_DECIMAL:	
+		TranslatedKeyCode = TEXT('.'); 
+		break;
+	case VK_DIVIDE:		
+		TranslatedKeyCode = TEXT('/'); 
+		break;
+
+	case VK_F1:
+	case VK_F2:
+	case VK_F3:
+	case VK_F4:
+	case VK_F5:
+	case VK_F6:
+	case VK_F7:
+	case VK_F8:
+	case VK_F9:
+	case VK_F10:
+	case VK_F11:
+	case VK_F12:
+	case VK_NUMLOCK:
+	case VK_SCROLL:
+	case VK_LSHIFT:
+	case VK_RSHIFT:
+	case VK_LCONTROL:
+	case VK_RCONTROL:
+	case VK_LMENU:
+	case VK_RMENU:
+		TranslatedKeyCode = 0;
+		break;
+
+	case 0xbb:
+		TranslatedKeyCode = (!bShiftIsDown) ? TEXT('=') : TEXT('+');
+		break;
+	case 0xbc:
+		TranslatedKeyCode = (!bShiftIsDown) ? TEXT(',') : TEXT('<');
+		break;
+	case 0xbd:
+		TranslatedKeyCode = (!bShiftIsDown) ? TEXT('-') : TEXT('_');
+		break;
+	case 0xbe:
+		TranslatedKeyCode = (!bShiftIsDown) ? TEXT('.') : TEXT('>');
+		break;
+	case 0xbf:
+		TranslatedKeyCode = (!bShiftIsDown) ? TEXT('/') : TEXT('?');
+		break;
+	case 0xc0:
+		TranslatedKeyCode = (!bShiftIsDown) ? TEXT('`') : TEXT('~');
+		break;
+	case 0xdb:
+		TranslatedKeyCode = (!bShiftIsDown) ? TEXT('[') : TEXT('{');
+		break;
+	case 0xdc:
+		TranslatedKeyCode = (!bShiftIsDown) ? TEXT('\\') : TEXT('|'); 
+		break;
+	case 0xdd:
+		TranslatedKeyCode = (!bShiftIsDown) ? TEXT(']') : TEXT('}');
+		break;
+	}
+
+	return TranslatedKeyCode;
+}
diff --git a/Engine/Platforms/Hololens/Source/Runtime/ApplicationCore/Private/HoloLensPlatformApplicationMisc.cpp b/Engine/Platforms/Hololens/Source/Runtime/ApplicationCore/Private/HoloLensPlatformApplicationMisc.cpp
new file mode 100644
--- /dev/null
+++ b/Engine/Platforms/Hololens/Source/Runtime/ApplicationCore/Private/HoloLensPlatformApplicationMisc.cpp
@@ -0,0 +1,77 @@
+// Copyright Epic Games, Inc. All Rights Reserved.
+
+#include "HoloLensPlatformApplicationMisc.h"
+#include "HoloLensApplication.h"
+#include "HoloLensErrorOutputDevice.h"
+#include "HAL/PlatformProcess.h"
+#include "Misc/App.h"
+
+FOutputDeviceError* FHoloLensPlatformApplicationMisc::GetErrorOutputDevice()
+{
+	static FHoloLensErrorOutputDevice Singleton;
+	return &Singleton;
+}
+
+GenericApplication* FHoloLensPlatformApplicationMisc::CreateApplication()
+{
+	return FHoloLensApplication::CreateHoloLensApplication();
+}
+
+// Defined in HoloLensLaunch.cpp
+extern void appWinPumpMessages();
+
+void FHoloLensPlatformApplicationMisc::PumpMessages(bool bFromMainLoop)
+{
+	if (!bFromMainLoop)
+	{
+		return;
+	}
+
+	GPumpingMessagesOutsideOfMainLoop = false;
+
+	appWinPumpMessages();
+
+	bool HasFocus = true;
+	// if its our window, allow sound, otherwise apply multiplier
+	FApp::SetVolumeMultiplier(HasFocus ? 1.0f : FApp::GetUnfocusedVolumeMultiplier());
+}
+
+float FHoloLensPlatformApplicationMisc::GetDPIScaleFactorAtPoint(float X, float Y)
+{
+	float Dpi = static_cast<uint32_t>(Windows::Graphics::Display::DisplayInformation::GetForCurrentView()->LogicalDpi);
+	return Dpi / 96.0f;
+}
+
+void FHoloLensPlatformApplicationMisc::ClipboardCopy(const TCHAR* Str)
+{
+	using namespace Windows::ApplicationModel::DataTransfer;
+
+	DataPackage ^Package = ref new DataPackage();
+	Package->RequestedOperation = DataPackageOperation::Copy;
+	Package->SetText(ref new Platform::String(Str));
+	Clipboard::SetContent(Package);
+}
+
+void FHoloLensPlatformApplicationMisc::ClipboardPaste(class FString& Dest)
+{
+	using namespace Windows::Foundation;
+	using namespace Windows::ApplicationModel::DataTransfer;
+
+	DataPackageView^ ClipboardContent = Clipboard::GetContent();
+	if (ClipboardContent != nullptr && ClipboardContent->Contains(StandardDataFormats::Text))
+	{
+		IAsyncOperation<Platform::String^>^ AsyncOp = ClipboardContent->GetTextAsync();
+		while (AsyncOp->Status == AsyncStatus::Started)
+		{
+			PumpMessages(false);
+
+			// Yield
+			FPlatformProcess::Sleep(0.f);
+		}
+
+		if (AsyncOp->Status == AsyncStatus::Completed)
+		{
+			Dest = AsyncOp->GetResults()->Data();
+		}
+	}
+}
diff --git a/Engine/Platforms/Hololens/Source/Runtime/ApplicationCore/Private/HoloLensPlatformSplash.cpp b/Engine/Platforms/Hololens/Source/Runtime/ApplicationCore/Private/HoloLensPlatformSplash.cpp
new file mode 100644
--- /dev/null
+++ b/Engine/Platforms/Hololens/Source/Runtime/ApplicationCore/Private/HoloLensPlatformSplash.cpp
@@ -0,0 +1,21 @@
+// Copyright Epic Games, Inc. All Rights Reserved.
+
+#include "HoloLensPlatformSplash.h"
+#include "HoloLensPlatformApplicationMisc.h"
+
+void FHoloLensSplash::Show()
+{
+	[REDACTED_EMAIL]: Implement me
+	FHoloLensPlatformApplicationMisc::PumpMessages(true);
+}
+
+void FHoloLensSplash::Hide()
+{
+	[REDACTED_EMAIL]: Implement me
+	FHoloLensPlatformApplicationMisc::PumpMessages(true);
+}
+
+void FHoloLensSplash::SetSplashText( const SplashTextType::Type InType, const TCHAR* InText )
+{
+	[REDACTED_EMAIL]: Implement me
+}
diff --git a/Engine/Platforms/Hololens/Source/Runtime/ApplicationCore/Private/HoloLensWindow.cpp b/Engine/Platforms/Hololens/Source/Runtime/ApplicationCore/Private/HoloLensWindow.cpp
new file mode 100644
--- /dev/null
+++ b/Engine/Platforms/Hololens/Source/Runtime/ApplicationCore/Private/HoloLensWindow.cpp
@@ -0,0 +1,110 @@
+// Copyright Epic Games, Inc. All Rights Reserved.
+
+
+#include "HoloLensWindow.h"
+#include "HoloLensApplication.h"
+#include "CoreTypes.h"
+
+FHoloLensWindow::FHoloLensWindow() :
+	WindowMode(EWindowMode::Windowed)
+{
+}
+
+FHoloLensWindow::~FHoloLensWindow()
+{
+}
+
+TSharedRef< FHoloLensWindow > FHoloLensWindow::Make()
+{
+	// First, allocate the new native window object.  This doesn't actually contain a native window or anything,
+	// we're simply instantiating the object so that we can keep shared references to it.
+	return MakeShareable( new FHoloLensWindow() );
+}
+
+void FHoloLensWindow::Initialize(FHoloLensApplication* const Application, const TSharedRef< FGenericWindowDefinition >& InDefinition)
+{
+	// Only 2 options in HoloLens: windowed and windowed fullscreen.
+	auto View = Windows::UI::ViewManagement::ApplicationView::GetForCurrentView();
+	WindowMode = View->IsFullScreenMode ? EWindowMode::WindowedFullscreen : EWindowMode::Windowed;
+}
+
+static const float DipsPerInch = 96.0f;
+
+int32 FHoloLensWindow::ConvertDipsToPixels(int32 Dips, float Dpi)
+{
+	return FMath::FloorToInt(static_cast<float>(Dips) * Dpi / DipsPerInch + 0.5f);
+}
+
+int32 FHoloLensWindow::ConvertPixelsToDips(int32 Pixels, float Dpi)
+{
+	return FMath::FloorToInt(static_cast<float>(Pixels) * DipsPerInch / Dpi + 0.5f);
+}
+
+void FHoloLensWindow::ReshapeWindow(int32 X, int32 Y, int32 Width, int32 Height)
+{
+	auto View = Windows::UI::ViewManagement::ApplicationView::GetForCurrentView();
+	float Dpi = static_cast<uint32_t>(Windows::Graphics::Display::DisplayInformation::GetForCurrentView()->LogicalDpi);
+	Windows::Foundation::Size RequestedSize;
+	RequestedSize.Width = ConvertPixelsToDips(Width, Dpi);
+	RequestedSize.Height = ConvertPixelsToDips(Height, Dpi);
+
+	// Not a big problem if this fails - will be handled by AdjustCachedSize
+	if (View->TryResizeView(RequestedSize))
+	{
+		// Remember for next time - when we launch windowed we try to match the window size to the
+		// render target size, which will always come through this call.
+		View->PreferredLaunchViewSize = RequestedSize;
+	}
+}
+
+FPlatformRect FHoloLensWindow::GetOSWindowBounds()
+{
+	auto LowLevelWindow = Windows::ApplicationModel::Core::CoreApplication::GetCurrentView()->CoreWindow;
+	float Dpi = static_cast<uint32_t>(Windows::Graphics::Display::DisplayInformation::GetForCurrentView()->LogicalDpi);
+	FPlatformRect Bounds;
+	Bounds.Top = ConvertDipsToPixels(LowLevelWindow->Bounds.Top, Dpi);
+	Bounds.Left = ConvertDipsToPixels(LowLevelWindow->Bounds.Left, Dpi);
+	Bounds.Bottom = ConvertDipsToPixels(LowLevelWindow->Bounds.Bottom, Dpi);
+	Bounds.Right = ConvertDipsToPixels(LowLevelWindow->Bounds.Right, Dpi);
+	return Bounds;
+}
+
+void FHoloLensWindow::AdjustCachedSize(FVector2D& Size) const
+{
+	// Force SWindow size to match the bounds of the OS window, which we are not entirely in control of.
+	FPlatformRect OSBounds = GetOSWindowBounds();
+	Size.X = OSBounds.Right - OSBounds.Left;
+	Size.Y = OSBounds.Bottom - OSBounds.Top;
+}
+
+void FHoloLensWindow::SetWindowMode(EWindowMode::Type InNewWindowMode)
+{
+	auto View = Windows::UI::ViewManagement::ApplicationView::GetForCurrentView();
+	if (InNewWindowMode == EWindowMode::Windowed && View->IsFullScreenMode)
+	{
+		View->ExitFullScreenMode();
+
+		// This will persist the user preference for Windows, and will attempt to match the size to the
+		// underlying render resolution.  Note it will ignore direct resizing of the window by the user
+		// (not typical HoloLens behavior, but probably desirable for UE games).
+		View->PreferredLaunchWindowingMode = Windows::UI::ViewManagement::ApplicationViewWindowingMode::PreferredLaunchViewSize;
+	}
+	else if (InNewWindowMode != EWindowMode::Windowed && !View->IsFullScreenMode)
+	{
+		if (View->TryEnterFullScreenMode())
+		{
+			// Persist the preference for full-screen.
+			View->PreferredLaunchWindowingMode = Windows::UI::ViewManagement::ApplicationViewWindowingMode::FullScreen;
+		}
+
+		// Note that the attempt to go full-screen might fail!
+		// @todo: who's responsible for checking this and handling it?
+	}
+
+	WindowMode = View->IsFullScreenMode ? EWindowMode::WindowedFullscreen : EWindowMode::Windowed;
+}
+
+EWindowMode::Type FHoloLensWindow::GetWindowMode() const
+{
+	return WindowMode;
+}
diff --git a/Engine/Platforms/Hololens/Source/Runtime/ApplicationCore/Private/WindowsGamingInputInterface.cpp b/Engine/Platforms/Hololens/Source/Runtime/ApplicationCore/Private/WindowsGamingInputInterface.cpp
new file mode 100644
--- /dev/null
+++ b/Engine/Platforms/Hololens/Source/Runtime/ApplicationCore/Private/WindowsGamingInputInterface.cpp
@@ -0,0 +1,458 @@
+// Copyright Epic Games, Inc. All Rights Reserved.
+
+#include "WindowsGamingInputInterface.h"
+#include "HAL/PlatformTime.h"
+#include <unordered_map>
+
+DECLARE_LOG_CATEGORY_EXTERN(GamepadSystem, Log, All);
+DEFINE_LOG_CATEGORY(GamepadSystem);
+
+using namespace Windows::Foundation;
+using namespace Windows::Gaming::Input;
+
+// Button flags assigned by Windows::Gaming::Input
+enum { GamingInputButtonMask = 0x3FFF };
+
+// Allocate flags for emulated/extended digital button functionality
+enum { GamepadButtonAux_LeftTrigger = (1 << 22) };
+enum { GamepadButtonAux_RightTrigger = (1 << 23) };
+enum { GamepadButtonAux_LeftStickUp = (1 << 24) };
+enum { GamepadButtonAux_LeftStickDown = (1 << 25) };
+enum { GamepadButtonAux_LeftStickLeft = (1 << 26) };
+enum { GamepadButtonAux_LeftStickRight = (1 << 27) };
+enum { GamepadButtonAux_RightStickUp = (1 << 28) };
+enum { GamepadButtonAux_RightStickDown = (1 << 29) };
+enum { GamepadButtonAux_RightStickLeft = (1 << 30) };
+enum { GamepadButtonAux_RightStickRight = (1 << 31) };
+
+// Values taken from XINPUT_GAMEPAD_LEFT_THUMB_DEADZONE, XINPUT_GAMEPAD_RIGHT_THUMB_DEADZONE, XINPUT_GAMEPAD_TRIGGER_THRESHOLD used by XInput.
+constexpr float GamepadLeftStickDeadzone = (7849.0f / 32768.0f);
+constexpr float GamepadRightStickDeadzone = (8689.0f / 32768.0f);
+constexpr float GamepadTriggerDeadzone = (30.0f / 255.0f);
+
+
+TSharedRef< WindowsGamingInputInterface > WindowsGamingInputInterface::Create(const TSharedRef< FGenericApplicationMessageHandler >& InMessageHandler)
+{
+	return MakeShareable(new WindowsGamingInputInterface(InMessageHandler));
+}
+
+WindowsGamingInputInterface::WindowsGamingInputInterface(const TSharedRef< FGenericApplicationMessageHandler > & InMessageHandler)
+	: MessageHandler(InMessageHandler)
+	, InitialRepeatDelay(0.2f)
+	, SubsequentRepeatDelay(0.1f)
+	, UpdateEnabled(true)
+	, UpdateChanged(false)
+{
+	PreviousTime = FPlatformTime::Seconds();
+
+	// Listen to Gamepad Added events
+	EventHandler<Gamepad^>^ GamepadAddedEvent = ref new EventHandler<Gamepad^>(
+		[this](Platform::Object^, Gamepad^ Gamepad)
+	{
+		OnGamepadAdded(Gamepad);
+	});
+
+	// Listen to Gamepad Removed events
+	EventHandler<Gamepad^>^ GamepadRemovedEvent = ref new EventHandler<Gamepad^>(
+		[this](Platform::Object^, Gamepad^ Gamepad)
+	{
+		OnGamepadRemoved(Gamepad);
+	});
+
+	GamepadAddedEventToken = Gamepad::GamepadAdded += GamepadAddedEvent;
+	GamepadRemovedEventToken = Gamepad::GamepadRemoved += GamepadRemovedEvent;
+}
+
+WindowsGamingInputInterface::~WindowsGamingInputInterface()
+{
+	Gamepad::GamepadAdded -= GamepadAddedEventToken;
+	Gamepad::GamepadRemoved -= GamepadRemovedEventToken;
+}
+
+void WindowsGamingInputInterface::SetMessageHandler(const TSharedRef< FGenericApplicationMessageHandler > & InMessageHandler)
+{
+	MessageHandler = InMessageHandler;
+}
+
+void WindowsGamingInputInterface::OnGamepadAdded(Windows::Gaming::Input::Gamepad^ Gamepad)
+{
+	int Id = InvalidGamepadDeviceId;
+
+	// find slot and perform validation
+	for (int i = 0; i < MaxSupportedGamepads; i++)
+	{
+		if (InvalidGamepadDeviceId != PadInfo[i].DeviceId)
+		{
+			if (Gamepad == PadInfo[i].Gamepad)
+			{
+				return;		// already registered
+			}
+		}
+		else if (InvalidGamepadDeviceId == Id)
+		{
+			Id = i;
+		}
+	}
+
+	// add new gamepad
+	if (InvalidGamepadDeviceId != Id)
+	{
+		UE_LOG(GamepadSystem, Log, TEXT("Gamepad 0x%p (id = %d) added."), (void *)Gamepad, Id);
+		PadInfo[Id].DeviceId = Id;
+		PadInfo[Id].Gamepad = Gamepad;
+	}
+	else
+	{
+		UE_LOG(GamepadSystem, Log, TEXT("Failed to add Gamepad 0x%p.  (No Id available)"), (void *)Gamepad);
+	}
+}
+
+void WindowsGamingInputInterface::OnGamepadRemoved(Windows::Gaming::Input::Gamepad^ Gamepad)
+{
+	// find slot and remove
+	for (int i = 0; i < MaxSupportedGamepads; i++)
+	{
+		if (InvalidGamepadDeviceId != PadInfo[i].DeviceId)
+		{
+			if (Gamepad == PadInfo[i].Gamepad)
+			{
+				UE_LOG(GamepadSystem, Log, TEXT("Gamepad 0x%p (id = %d) removed."), (void *)Gamepad, PadInfo[i].DeviceId);
+				PadInfo[i].DeviceId = InvalidGamepadDeviceId;
+				PadInfo[i].Gamepad = nullptr;
+				PadInfo[i].TerminateInputs = true;
+				return;
+			}
+		}
+	}
+
+	UE_LOG(GamepadSystem, Log, TEXT("Failed to remove Gamepad 0x%p.  (not found)"), (void *)Gamepad);
+}
+
+Windows::Gaming::Input::Gamepad^ WindowsGamingInputInterface::GetGamepadForIndex(int Id)
+{
+	return ((((unsigned int)Id) < MaxSupportedGamepads) ? PadInfo[Id].Gamepad : nullptr);
+}
+
+bool WindowsGamingInputInterface::GetVibration(int Id, float& LeftMotor, float& RightMotor, float& LeftTrigger, float& RightTrigger)
+{
+	Gamepad^ GamePad = GetGamepadForIndex(Id);
+	if (GamePad == nullptr)
+	{
+		return false;
+	}
+
+	LeftMotor = PadInfo[Id].Gamepad->Vibration.LeftMotor;
+	RightMotor = PadInfo[Id].Gamepad->Vibration.RightMotor;
+	LeftTrigger = PadInfo[Id].Gamepad->Vibration.LeftTrigger;
+	RightTrigger = PadInfo[Id].Gamepad->Vibration.RightTrigger;
+
+	return true;
+}
+
+bool WindowsGamingInputInterface::SetVibration(int Id, float LeftMotor, float RightMotor, float LeftTrigger, float RightTrigger)
+{
+	Gamepad^ GamePad = GetGamepadForIndex(Id);
+	if (GamePad == nullptr)
+	{
+		return false;
+	}
+
+	if (PadInfo[Id].Gamepad->Vibration.LeftMotor != LeftMotor ||
+		PadInfo[Id].Gamepad->Vibration.RightMotor != RightMotor ||
+		PadInfo[Id].Gamepad->Vibration.LeftTrigger != LeftTrigger ||
+		PadInfo[Id].Gamepad->Vibration.RightTrigger != RightTrigger)
+	{
+		Windows::Gaming::Input::GamepadVibration Vibration;
+		Vibration.LeftMotor = FMath::Clamp(LeftMotor, 0.f, 1.f);
+		Vibration.RightMotor = FMath::Clamp(RightMotor, 0.f, 1.f);
+		Vibration.LeftTrigger = FMath::Clamp(LeftTrigger, 0.f, 1.f);
+		Vibration.RightTrigger = FMath::Clamp(RightTrigger, 0.f, 1.f);
+
+		UE_LOG(GamepadSystem, Verbose, TEXT("Set Gamepad 0x%p (id = %d) vibration. (%f, %f, %f, %f)"), (void *)(PadInfo[Id].Gamepad), Id, Vibration.LeftMotor, Vibration.RightMotor, Vibration.LeftTrigger, Vibration.RightTrigger);
+		PadInfo[Id].Gamepad->Vibration = Vibration;
+	}
+
+	return true;
+}
+
+static bool GamepadButtonToUnrealName(FGamepadKeyNames::Type& ButtonKey, uint32 ButtonMask)
+{
+	typedef std::unordered_map< uint32, FGamepadKeyNames::Type > ButtonMapType;
+	static const ButtonMapType ButtonMap
+	{
+		{ (uint32)Windows::Gaming::Input::GamepadButtons::A,				FGamepadKeyNames::FaceButtonBottom },
+		{ (uint32)Windows::Gaming::Input::GamepadButtons::B,				FGamepadKeyNames::FaceButtonRight },
+		{ (uint32)Windows::Gaming::Input::GamepadButtons::X,				FGamepadKeyNames::FaceButtonLeft },
+		{ (uint32)Windows::Gaming::Input::GamepadButtons::Y,				FGamepadKeyNames::FaceButtonTop },
+		{ (uint32)Windows::Gaming::Input::GamepadButtons::LeftShoulder,		FGamepadKeyNames::LeftShoulder },
+		{ (uint32)Windows::Gaming::Input::GamepadButtons::RightShoulder,	FGamepadKeyNames::RightShoulder },
+		{ (uint32)Windows::Gaming::Input::GamepadButtons::Menu,				FGamepadKeyNames::SpecialRight },
+		{ (uint32)Windows::Gaming::Input::GamepadButtons::View,				FGamepadKeyNames::SpecialLeft },
+		{ (uint32)Windows::Gaming::Input::GamepadButtons::DPadUp,			FGamepadKeyNames::DPadUp },
+		{ (uint32)Windows::Gaming::Input::GamepadButtons::DPadDown,			FGamepadKeyNames::DPadDown },
+		{ (uint32)Windows::Gaming::Input::GamepadButtons::DPadLeft,			FGamepadKeyNames::DPadLeft },
+		{ (uint32)Windows::Gaming::Input::GamepadButtons::DPadRight,		FGamepadKeyNames::DPadRight },
+		{ (uint32)Windows::Gaming::Input::GamepadButtons::LeftThumbstick,	FGamepadKeyNames::LeftThumb },
+		{ (uint32)Windows::Gaming::Input::GamepadButtons::RightThumbstick,	FGamepadKeyNames::RightThumb },
+
+		{ GamepadButtonAux_LeftTrigger,			FGamepadKeyNames::LeftTriggerThreshold },
+		{ GamepadButtonAux_RightTrigger,		FGamepadKeyNames::RightTriggerThreshold },
+		{ GamepadButtonAux_LeftStickUp,			FGamepadKeyNames::LeftStickUp },
+		{ GamepadButtonAux_LeftStickDown,		FGamepadKeyNames::LeftStickDown },
+		{ GamepadButtonAux_LeftStickLeft,		FGamepadKeyNames::LeftStickLeft },
+		{ GamepadButtonAux_LeftStickRight,		FGamepadKeyNames::LeftStickRight },
+		{ GamepadButtonAux_RightStickUp,		FGamepadKeyNames::RightStickUp },
+		{ GamepadButtonAux_RightStickDown,		FGamepadKeyNames::RightStickDown },
+		{ GamepadButtonAux_RightStickLeft,		FGamepadKeyNames::RightStickLeft },
+		{ GamepadButtonAux_RightStickRight,		FGamepadKeyNames::RightStickRight }
+	};
+
+	ButtonMapType::const_iterator it = ButtonMap.find(ButtonMask);
+	if (ButtonMap.end() != it)
+	{
+		ButtonKey = it->second;
+		return true;
+	}
+
+	return false;
+}
+
+void WindowsGamingInputInterface::EnableUpdate()
+{
+	if (!UpdateEnabled)
+	{
+		UpdateEnabled = true;
+		UpdateChanged = true;
+	}
+}
+
+void WindowsGamingInputInterface::DisableUpdate()
+{
+	if (UpdateEnabled)
+	{
+		UpdateEnabled = false;
+		UpdateChanged = true;
+	}
+}
+
+void WindowsGamingInputInterface::TerminateGamepadInputs(GamepadMapping& Mapping, int Id)
+{
+	UE_LOG(GamepadSystem, Log, TEXT("Terminate active gamepad inputs."));
+
+	// reset axes
+	MessageHandler->OnControllerAnalog(FGamepadKeyNames::LeftAnalogX, Id, 0.0f);
+	MessageHandler->OnControllerAnalog(FGamepadKeyNames::LeftAnalogY, Id, 0.0f);
+	MessageHandler->OnControllerAnalog(FGamepadKeyNames::RightAnalogX, Id, 0.0f);
+	MessageHandler->OnControllerAnalog(FGamepadKeyNames::RightAnalogY, Id, 0.0f);
+
+	if (Mapping.LeftTrigger > 0)
+	{
+		MessageHandler->OnControllerAnalog(FGamepadKeyNames::LeftTriggerAnalog, Id, 0.0f);
+		Mapping.LeftTrigger = 0.0f;
+	}
+	if (Mapping.RightTrigger > 0)
+	{
+		MessageHandler->OnControllerAnalog(FGamepadKeyNames::RightTriggerAnalog, Id, 0.0f);
+		Mapping.RightTrigger = 0.0f;
+	}
+
+	unsigned int ButtonMask = Mapping.ButtonHeldMask;
+	unsigned int BitMask = 1;
+	while ((0 != ButtonMask) && (0 != BitMask))
+	{
+		FName ButtonKey;
+
+		if (0 != (ButtonMask & BitMask))
+		{
+			if (GamepadButtonToUnrealName(ButtonKey, BitMask))
+			{
+				MessageHandler->OnControllerButtonReleased(ButtonKey, Id, false);
+			}
+		}
+		BitMask <<= 1;
+	}
+
+	Mapping.ButtonHeldMask = 0;
+	Mapping.TerminateInputs = false;
+
+	for (uint32 i = 0; i < MaxSupportedButtons; i++)
+	{
+		Mapping.RepeatTime[i] = 0.0f;
+	}
+}
+
+void WindowsGamingInputInterface::UpdateGamepads()
+{
+	const double CurrentTime = FPlatformTime::Seconds();
+	const float  DeltaTime = (float)(CurrentTime - PreviousTime);
+	PreviousTime = CurrentTime;
+
+	// handle update toggle
+	if (UpdateChanged)
+	{
+		if (UpdateEnabled)
+		{
+			// terminate existing inputs (e.g.  generate button up events for pressed buttons)
+			for (int i = 0; i < MaxSupportedGamepads; i++)
+			{
+				if ((PadInfo[i].TerminateInputs) ||
+					((InvalidGamepadDeviceId != PadInfo[i].DeviceId) &&
+						(nullptr != PadInfo[i].Gamepad)))
+				{
+					int	Id = i;
+					TerminateGamepadInputs(PadInfo[i], Id);
+				}
+			}
+		}
+		UpdateChanged = false;
+	}
+
+	if (UpdateEnabled)
+	{
+		// find slot and remove
+		for (int i = 0; i < MaxSupportedGamepads; i++)
+		{
+			if ((InvalidGamepadDeviceId != PadInfo[i].DeviceId) &&
+				(nullptr != PadInfo[i].Gamepad))
+			{
+				GamepadReading Reading = PadInfo[i].Gamepad->GetCurrentReading();
+				float LeftThumbX = (float)Reading.LeftThumbstickX;
+				float LeftThumbY = (float)Reading.LeftThumbstickY;
+				float RightThumbX = (float)Reading.RightThumbstickX;
+				float RightThumbY = (float)Reading.RightThumbstickY;
+				float LeftTrigger = (float)Reading.LeftTrigger;
+				float RightTrigger = (float)Reading.RightTrigger;
+
+				// map axes
+				MessageHandler->OnControllerAnalog(FGamepadKeyNames::LeftAnalogX, i, LeftThumbX);
+				MessageHandler->OnControllerAnalog(FGamepadKeyNames::LeftAnalogY, i, LeftThumbY);
+				MessageHandler->OnControllerAnalog(FGamepadKeyNames::RightAnalogX, i, RightThumbX);
+				MessageHandler->OnControllerAnalog(FGamepadKeyNames::RightAnalogY, i, RightThumbY);
+
+				// map analog triggers directly
+				MessageHandler->OnControllerAnalog(FGamepadKeyNames::LeftTriggerAnalog, i, LeftTrigger);
+				MessageHandler->OnControllerAnalog(FGamepadKeyNames::RightTriggerAnalog, i, RightTrigger);
+
+				// map buttons (low 15 bits)
+				unsigned int CurrentButtonHeldMask = (((unsigned int)Reading.Buttons) & GamingInputButtonMask);
+
+				// map analog triggers to digital input.
+				float PrevLeftTrigger = PadInfo[i].LeftTrigger;
+				if (LeftTrigger > GamepadTriggerDeadzone)
+				{
+					CurrentButtonHeldMask |= GamepadButtonAux_LeftTrigger;
+				}
+				PadInfo[i].LeftTrigger = LeftTrigger;
+
+				float PrevRightTrigger = PadInfo[i].RightTrigger;
+				if (RightTrigger > GamepadTriggerDeadzone)
+				{
+					CurrentButtonHeldMask |= GamepadButtonAux_RightTrigger;
+				}
+				PadInfo[i].RightTrigger = RightTrigger;
+
+				// map left/right stick digital inputs to (top 8 bits)
+
+				// kept h/v key switch due to following change.
+				// @BLACK_TUSK_CHANGE: [REDACTED_EMAIL] - BEGIN - Only send a single direction of left stick analog readings as a key down event at a time.
+				// - This is to work around the fact that SWidget navigation for the left stick responds to OnKeyDown rather than OnAnalogValue
+				// - To get around that without this change would require a number of changes to the way SWidget handles navigation, 
+				// - and not only that, but any of our UMG screens that implement OnPreviewKeyDown or OnKeyDown themselves to handle left stick navigation manually would need changes as well.
+				// - UDN post: https://udn.unrealengine.com/questions/278514/overly-sensitive-umg-navigation.html
+				const bool UseLeftstickVerticalReading = FMath::Abs(LeftThumbY) >= FMath::Abs(LeftThumbX);
+				if (UseLeftstickVerticalReading)
+				{
+					if (LeftThumbY > GamepadLeftStickDeadzone)
+					{
+						CurrentButtonHeldMask |= GamepadButtonAux_LeftStickUp;
+					}
+					else if (LeftThumbY < -GamepadLeftStickDeadzone)
+					{
+						CurrentButtonHeldMask |= GamepadButtonAux_LeftStickDown;
+					}
+				}
+				else
+				{
+					if (LeftThumbX > GamepadLeftStickDeadzone)
+					{
+						CurrentButtonHeldMask |= GamepadButtonAux_LeftStickRight;
+					}
+					else if (LeftThumbX < -GamepadLeftStickDeadzone)
+					{
+						CurrentButtonHeldMask |= GamepadButtonAux_LeftStickLeft;
+					}
+				}
+
+				if (RightThumbY > GamepadRightStickDeadzone)
+				{
+					CurrentButtonHeldMask |= GamepadButtonAux_RightStickUp;
+				}
+				else if (RightThumbY < -GamepadRightStickDeadzone)
+				{
+					CurrentButtonHeldMask |= GamepadButtonAux_RightStickDown;
+				}
+				if (RightThumbX > GamepadRightStickDeadzone)
+				{
+					CurrentButtonHeldMask |= GamepadButtonAux_RightStickRight;
+				}
+				else if (RightThumbX < -GamepadRightStickDeadzone)
+				{
+					CurrentButtonHeldMask |= GamepadButtonAux_RightStickLeft;
+				}
+
+				// handle button change events
+				unsigned int ActionMask = (PadInfo[i].ButtonHeldMask ^ CurrentButtonHeldMask);
+				unsigned int RepeatMask = (PadInfo[i].ButtonHeldMask & CurrentButtonHeldMask);
+				unsigned int BitMask = 1;
+
+				for (uint32 n = 0; n < MaxSupportedButtons; ++n)
+				{
+					FGamepadKeyNames::Type ButtonKey;
+
+					// check for button state change
+					if (0 != (ActionMask & BitMask))
+					{
+						if (0 != (CurrentButtonHeldMask & BitMask))
+						{
+							if (GamepadButtonToUnrealName(ButtonKey, BitMask))
+							{
+								UE_LOG(GamepadSystem, Verbose, TEXT("Gamepad 0x%p (id = %d) - %s Pressed"), (void *)(PadInfo[i].Gamepad), i, *ButtonKey.ToString());
+								MessageHandler->OnControllerButtonPressed(ButtonKey, i, false);
+								PadInfo[i].RepeatTime[n] = InitialRepeatDelay;
+							}
+						}
+						else
+						{
+							if (GamepadButtonToUnrealName(ButtonKey, BitMask))
+							{
+								UE_LOG(GamepadSystem, Verbose, TEXT("Gamepad 0x%p (id = %d) - %s Released"), (void *)(PadInfo[i].Gamepad), i, *ButtonKey.ToString());
+								MessageHandler->OnControllerButtonReleased(ButtonKey, i, false);
+								PadInfo[i].RepeatTime[n] = 0.0f;
+							}
+						}
+					}
+
+					// check for repeat key
+					if (0 != (RepeatMask & BitMask))
+					{
+						PadInfo[i].RepeatTime[n] -= DeltaTime;
+						if (PadInfo[i].RepeatTime[n] <= 0.0f)
+						{
+							PadInfo[i].RepeatTime[n] = SubsequentRepeatDelay;
+							if (GamepadButtonToUnrealName(ButtonKey, BitMask))
+							{
+								UE_LOG(GamepadSystem, Verbose, TEXT("Gamepad 0x%p (id = %d) - %s Repeated"), (void *)(PadInfo[i].Gamepad), i, *ButtonKey.ToString());
+								MessageHandler->OnControllerButtonPressed(ButtonKey, i, true);
+							}
+						}
+					}
+
+					BitMask <<= 1;
+				}
+
+				// update button states
+				PadInfo[i].ButtonHeldMask = CurrentButtonHeldMask;
+			}
+		}
+	}
+}
diff --git a/Engine/Platforms/Hololens/Source/Runtime/ApplicationCore/Private/WindowsGamingInputInterface.h b/Engine/Platforms/Hololens/Source/Runtime/ApplicationCore/Private/WindowsGamingInputInterface.h
new file mode 100644
--- /dev/null
+++ b/Engine/Platforms/Hololens/Source/Runtime/ApplicationCore/Private/WindowsGamingInputInterface.h
@@ -0,0 +1,81 @@
+// Copyright Epic Games, Inc. All Rights Reserved.
+
+#pragma once
+#include "GenericPlatform/GenericApplicationMessageHandler.h"
+
+class WindowsGamingInputInterface
+{
+public:
+
+	static TSharedRef< WindowsGamingInputInterface > Create(const TSharedRef< FGenericApplicationMessageHandler >& InMessageHandler);
+
+	~WindowsGamingInputInterface();
+
+	void SetMessageHandler(const TSharedRef< FGenericApplicationMessageHandler > & InMessageHandler);
+
+	void UpdateGamepads();
+
+	bool GetVibration(int Id, float& LeftMotor, float& RightMotor, float& LeftTrigger, float& RightTrigger);
+
+	bool SetVibration(int Id, float LeftMotor, float RightMotor, float LeftTrigger, float RightTrigger);
+
+	void EnableUpdate();
+
+	void DisableUpdate();
+
+	Windows::Gaming::Input::Gamepad^ GetGamepadForIndex(int Id);
+
+protected:
+
+	void OnGamepadAdded(Windows::Gaming::Input::Gamepad^ Gamepad);
+	void OnGamepadRemoved(Windows::Gaming::Input::Gamepad^ Gamepad);
+
+private:
+
+	WindowsGamingInputInterface(const TSharedRef< FGenericApplicationMessageHandler >& MessageHandler);
+
+private:
+
+	TSharedRef< FGenericApplicationMessageHandler > MessageHandler;
+	Windows::Foundation::EventRegistrationToken		GamepadAddedEventToken;
+	Windows::Foundation::EventRegistrationToken		GamepadRemovedEventToken;
+	double											PreviousTime;
+	float											InitialRepeatDelay;
+	float											SubsequentRepeatDelay;
+
+	enum { MaxSupportedGamepads = 4 };		// MAX_NUM_XINPUT_CONTROLLER == 4 (in XInput manager)
+	enum { InvalidGamepadDeviceId = -1 };
+	enum { MaxSupportedButtons = 32 };
+
+	class GamepadMapping
+	{
+	public:
+		int									DeviceId;
+		Windows::Gaming::Input::Gamepad^	Gamepad;
+		unsigned int						ButtonHeldMask;
+		float								LeftTrigger;
+		float								RightTrigger;
+		bool								TerminateInputs;
+		float								RepeatTime[MaxSupportedButtons];
+
+		GamepadMapping()
+			: DeviceId(InvalidGamepadDeviceId)
+			, ButtonHeldMask(0)
+			, LeftTrigger(0.0f)
+			, RightTrigger(0.0f)
+			, TerminateInputs(false)
+		{
+			for (uint32 i = 0; i < MaxSupportedButtons; i++)
+			{
+				RepeatTime[i] = 0.0f;
+			}
+		}
+	};
+
+	GamepadMapping		PadInfo[MaxSupportedGamepads];
+	bool				UpdateEnabled;
+	bool				UpdateChanged;
+
+	void TerminateGamepadInputs(GamepadMapping& Mapping, int Id);
+};
+
diff --git a/Engine/Platforms/Hololens/Source/Runtime/ApplicationCore/Public/HoloLensApplication.h b/Engine/Platforms/Hololens/Source/Runtime/ApplicationCore/Public/HoloLensApplication.h
new file mode 100644
--- /dev/null
+++ b/Engine/Platforms/Hololens/Source/Runtime/ApplicationCore/Public/HoloLensApplication.h
@@ -0,0 +1,85 @@
+// Copyright Epic Games, Inc. All Rights Reserved.
+
+#pragma once
+
+#include "GenericPlatform/GenericApplication.h"
+
+DECLARE_LOG_CATEGORY_EXTERN(LogHoloLens, Log, All);
+
+class FHoloLensApplication : public GenericApplication
+{
+public:
+
+	static FHoloLensApplication* CreateHoloLensApplication();
+
+	static FHoloLensApplication* GetHoloLensApplication();
+
+	static void CacheDesktopSize();
+
+	static FVector2D GetDesktopSize();
+
+public:	
+
+	virtual ~FHoloLensApplication() {}
+
+	virtual void PollGameDeviceState( const float TimeDelta ) override;
+
+	virtual bool IsMouseAttached() const override;
+
+	virtual bool IsGamepadAttached() const override;
+
+	virtual FPlatformRect GetWorkArea( const FPlatformRect& CurrentWindow ) const override;
+	virtual TSharedRef< FGenericWindow > MakeWindow();
+	virtual void InitializeWindow(const TSharedRef< FGenericWindow >& Window, const TSharedRef< FGenericWindowDefinition >& InDefinition, const TSharedPtr< FGenericWindow >& InParent, const bool bShowImmediately);
+
+	virtual IInputInterface* GetInputInterface() override;
+
+	TSharedRef< class FGenericApplicationMessageHandler > GetMessageHandler() const;
+
+	void SetMessageHandler( const TSharedRef< FGenericApplicationMessageHandler >& InMessageHandler ) override;
+	virtual void PumpMessages(const float TimeDelta) override;
+
+	virtual void GetInitialDisplayMetrics(FDisplayMetrics& OutDisplayMetrics) const override;
+
+	TSharedRef< class FHoloLensCursor > GetCursor() const;
+	TSharedPtr< class FHoloLensInputInterface > GetHoloLensInputInterface() const { return InputInterface; }
+
+	TSharedPtr< class FHoloLensWindow > GetHoloLensWindow() const { return ApplicationWindow; }
+
+	virtual bool ApplicationLicenseValid(FPlatformUserId PlatformUser = PLATFORMUSERID_NONE);
+
+	// @MIXEDREALITY_CHANGE : BEGIN
+	virtual void AddExternalInputDevice(TSharedPtr< class IInputDevice > InputDevice);
+	// @MIXEDREALITY_CHANGE : END
+	
+#if PLATFORM_HOLOLENS
+public:
+	static void SetHolographicSpace(Windows::Graphics::Holographic::HolographicSpace^ holoSpace) { HoloSpace = holoSpace; }
+	static Windows::Graphics::Holographic::HolographicSpace^ GetHolographicSpace() { return HoloSpace; }
+private:
+	static Windows::Graphics::Holographic::HolographicSpace^ HoloSpace;
+#endif
+
+private:
+
+	FHoloLensApplication();
+
+	static void InitLicensing();
+	static void LicenseChangedHandler();
+	static void BroadcastUELicenseChangeEvent();
+private:
+
+	// @MIXEDREALITY_CHANGE : BEGIN
+	/** List of input devices implemented in external modules. **/
+	TArray< TSharedPtr< class IInputDevice > > ExternalInputDevices;
+	bool bHasLoadedInputPlugins;
+	// @MIXEDREALITY_CHANGE : END
+
+	TSharedPtr< class FHoloLensInputInterface > InputInterface;
+
+	TSharedRef< class FHoloLensWindow > ApplicationWindow;
+
+	static FVector2D DesktopSize;
+	
+	static bool buildForRetailWindowsStore;
+};
\ No newline at end of file
diff --git a/Engine/Platforms/Hololens/Source/Runtime/ApplicationCore/Public/HoloLensCursor.h b/Engine/Platforms/Hololens/Source/Runtime/ApplicationCore/Public/HoloLensCursor.h
new file mode 100644
--- /dev/null
+++ b/Engine/Platforms/Hololens/Source/Runtime/ApplicationCore/Public/HoloLensCursor.h
@@ -0,0 +1,77 @@
+// Copyright Epic Games, Inc. All Rights Reserved.
+
+#pragma once
+
+#include "GenericPlatform/ICursor.h"
+#include "Math/IntVector.h"
+//#include "Microsoft/AllowMicrosoftPlatformTypes.h"
+//#include "agile.h"
+//#include "Microsoft/HideMicrosoftPlatformTypes.h"
+
+
+ref class FHoloLensCursorMouseEventObj sealed
+{
+public:
+    FHoloLensCursorMouseEventObj();
+
+    void OnMouseMoved(Windows::Devices::Input::MouseDevice ^sender, Windows::Devices::Input::MouseEventArgs ^args);
+
+    Windows::Foundation::TypedEventHandler<Windows::Devices::Input::MouseDevice ^, Windows::Devices::Input::MouseEventArgs ^>^ GetMouseMovedHandler();
+};
+
+
+class FHoloLensCursor : public ICursor
+{
+public:
+
+    FHoloLensCursor();
+
+	virtual ~FHoloLensCursor();
+
+	virtual FVector2D GetPosition() const override;
+
+	virtual void SetPosition( const int32 X, const int32 Y ) override;
+
+	virtual void SetType( const EMouseCursor::Type InNewCursor ) override;
+
+	virtual EMouseCursor::Type GetType() const override
+	{
+		return CurrentCursor;
+	}
+
+	virtual void GetSize( int32& Width, int32& Height ) const override;
+
+	virtual void Show( bool bShow ) override;
+
+	virtual void Lock( const RECT* const Bounds ) override;
+
+    void UpdatePosition(const FVector2D& NewPosition);
+
+    bool IsUsingRawMouseNoCursor() { return bUsingRawMouseNoCursor; }
+
+    void ProcessDeferredActions();
+
+	void OnRawMouseMove(const FIntVector& MouseDelta);
+
+	/**
+	* Allows overriding the shape of a particular cursor.
+	*/
+	void SetTypeShape(EMouseCursor::Type InCursorType, void* CursorHandle);
+
+private:
+
+	void SetUseRawMouse(bool bUse);
+
+	Windows::UI::Core::CoreCursor^ GetDefaultCursorForType(EMouseCursor::Type InCursorType);
+
+    EMouseCursor::Type                                CurrentCursor = (EMouseCursor::Type) - 1;
+    FVector2D                                         CursorPosition;
+    bool                                              bUsingRawMouseNoCursor;
+    bool                                              bDeferredCursorTypeChange;
+	TArray<FIntVector>                                DeferredMoveEvents;
+    
+    /** Cursors */
+    Platform::Array<Windows::UI::Core::CoreCursor^>^  Cursors;
+    FHoloLensCursorMouseEventObj^                          MouseEventObj;
+    Windows::Foundation::EventRegistrationToken       MouseEventRegistrationToken;
+};
\ No newline at end of file
diff --git a/Engine/Platforms/Hololens/Source/Runtime/ApplicationCore/Public/HoloLensErrorOutputDevice.h b/Engine/Platforms/Hololens/Source/Runtime/ApplicationCore/Public/HoloLensErrorOutputDevice.h
new file mode 100644
--- /dev/null
+++ b/Engine/Platforms/Hololens/Source/Runtime/ApplicationCore/Public/HoloLensErrorOutputDevice.h
@@ -0,0 +1,27 @@
+// Copyright Epic Games, Inc. All Rights Reserved.
+
+#pragma once
+
+#include "CoreTypes.h"
+#include "Misc/OutputDeviceError.h"
+
+class FHoloLensErrorOutputDevice : public FOutputDeviceError
+{
+public:
+	/** Constructor, initializing member variables */
+	APPLICATIONCORE_API FHoloLensErrorOutputDevice();
+
+	/**
+	* Serializes the passed in data unless the current event is suppressed.
+	*
+	* @param	Data	Text to log
+	* @param	Event	Event name used for suppression purposes
+	*/
+	virtual void Serialize(const TCHAR* Msg, ELogVerbosity::Type Verbosity, const class FName& Category) override;
+
+	/**
+	* Error handling function that is being called from within the system wide global
+	* error handler, e.g. using structured exception handling on the PC.
+	*/
+	void HandleError();
+};
\ No newline at end of file
diff --git a/Engine/Platforms/Hololens/Source/Runtime/ApplicationCore/Public/HoloLensInputInterface.h b/Engine/Platforms/Hololens/Source/Runtime/ApplicationCore/Public/HoloLensInputInterface.h
new file mode 100644
--- /dev/null
+++ b/Engine/Platforms/Hololens/Source/Runtime/ApplicationCore/Public/HoloLensInputInterface.h
@@ -0,0 +1,87 @@
+// Copyright Epic Games, Inc. All Rights Reserved.
+
+#pragma once
+
+#include "GenericPlatform/IInputInterface.h"
+#include "GenericPlatform/GenericApplicationMessageHandler.h"
+#include "Math/Color.h"
+
+/**
+ * Interface class for input devices.
+ */
+class FHoloLensInputInterface
+	: public IInputInterface
+{
+public:
+
+	static TSharedRef< FHoloLensInputInterface > Create( const TSharedRef< FGenericApplicationMessageHandler >& InMessageHandler );
+
+public:
+
+	~FHoloLensInputInterface();
+
+	void SetMessageHandler( const TSharedRef< FGenericApplicationMessageHandler >& InMessageHandler );
+
+	/** Tick the interface (i.e check for new controllers) */
+	void Tick( float DeltaTime );
+
+	/** Poll for controller state and send events if needed. */
+	void SendControllerEvents();
+
+	/** Sets the strength/speed of the given channel for the given controller id. */
+	virtual void SetForceFeedbackChannelValue(int32 ControllerId, FForceFeedbackChannelType ChannelType, float Value) override;
+
+	/** Sets the strength/speed of all the channels for the given controller id. */
+	virtual void SetForceFeedbackChannelValues(int32 ControllerId, const FForceFeedbackValues &Values) override;
+
+	/* Ignored if controller does not support a color. */
+	virtual void SetLightColor(int32 ControllerId, struct FColor Color) override { /* nop */ }
+
+	virtual void ResetLightColor(int32 ControllerId) override {}
+
+
+public:
+
+	/** Platform-specific code can call this function to get the Gamepad^ associated with a UserId, if any. */
+	//Windows::Gaming::Input::Gamepad^ GetGamepadForUser(const int UserId);
+
+	/** Platform-specific code can call this function to get the UserId associated with a Gamepad^, if any. */
+	int GetUserIdForController(Windows::Gaming::Input::Gamepad^ Controller);
+
+	/** Platform-specific backdoor to get the FPlatformUserId from an XboxUserId */
+	//FPlatformUserId GetPlatformUserIdFromXboxUserId(const TCHAR* XboxUserId);
+
+	/** Platform-specific code can call this function to get the Controller^ associated with a ControllerIndex, if any. */
+	Windows::Gaming::Input::Gamepad^ GetGamepadForControllerId(int32 ControllerId);
+
+private:
+
+	FHoloLensInputInterface( const TSharedRef< FGenericApplicationMessageHandler >& MessageHandler );
+
+	/**
+	 * Called when the application has been reactivated.
+	 */
+	void OnFocusGain();
+
+	/**
+	 * This is called when the application is about to be deactivated.
+	 */
+	void OnFocusLost();
+
+	/**
+	 *	Scans for keyboard input
+	 */
+	void ConditionalScanForKeyboardChanges( float DeltaTime );
+
+	/**
+	 *	Convert the given keycode to the character it represents
+	 */
+	uint32 MapVirtualKeyToCharacter(uint32 InVirtualKey, bool bShiftIsDown);
+
+
+private:
+
+	TSharedRef< FGenericApplicationMessageHandler > MessageHandler;
+
+	TSharedRef<class WindowsGamingInputInterface> GamingInput;
+};
\ No newline at end of file
diff --git a/Engine/Platforms/Hololens/Source/Runtime/ApplicationCore/Public/HoloLensPlatformApplicationMisc.h b/Engine/Platforms/Hololens/Source/Runtime/ApplicationCore/Public/HoloLensPlatformApplicationMisc.h
new file mode 100644
--- /dev/null
+++ b/Engine/Platforms/Hololens/Source/Runtime/ApplicationCore/Public/HoloLensPlatformApplicationMisc.h
@@ -0,0 +1,32 @@
+// Copyright Epic Games, Inc. All Rights Reserved.
+
+#pragma once
+
+#include "GenericPlatform/GenericPlatformApplicationMisc.h"
+
+struct APPLICATIONCORE_API FHoloLensPlatformApplicationMisc : public FGenericPlatformApplicationMisc
+{
+	static class FOutputDeviceError* GetErrorOutputDevice();
+	static class GenericApplication* CreateApplication();
+	static void PumpMessages(bool bFromMainLoop);
+	static float GetDPIScaleFactorAtPoint(float X, float Y);
+	static void ClipboardCopy(const TCHAR* Str);
+	static void ClipboardPaste(class FString& Dest);
+
+	static bool AnchorWindowWindowPositionTopLeft()
+	{
+		// UE expects mouse coordinates in screen space. HoloLens provides in client space. 
+		// Also note comments in FDisplayMetrics::GetDisplayMetrics for HoloLens.
+		return true;
+	}
+	
+	static bool RequiresVirtualKeyboard()
+	{
+#if PLATFORM_HOLOLENS
+		return true;
+#endif
+		return false;
+	}
+};
+
+typedef FHoloLensPlatformApplicationMisc FPlatformApplicationMisc;
diff --git a/Engine/Platforms/Hololens/Source/Runtime/ApplicationCore/Public/HoloLensPlatformSurvey.h b/Engine/Platforms/Hololens/Source/Runtime/ApplicationCore/Public/HoloLensPlatformSurvey.h
new file mode 100644
--- /dev/null
+++ b/Engine/Platforms/Hololens/Source/Runtime/ApplicationCore/Public/HoloLensPlatformSurvey.h
@@ -0,0 +1,8 @@
+// Copyright Epic Games, Inc. All Rights Reserved.
+
+#pragma once
+
+#include "GenericPlatform/GenericPlatformSurvey.h"
+
+// generic implementation
+typedef FGenericPlatformSurvey FPlatformSurvey;
diff --git a/Engine/Platforms/Hololens/Source/Runtime/ApplicationCore/Public/HoloLensSurvey.h b/Engine/Platforms/Hololens/Source/Runtime/ApplicationCore/Public/HoloLensSurvey.h
new file mode 100644
--- /dev/null
+++ b/Engine/Platforms/Hololens/Source/Runtime/ApplicationCore/Public/HoloLensSurvey.h
@@ -0,0 +1,18 @@
+// Copyright Epic Games, Inc. All Rights Reserved.
+
+/*=============================================================================================
+	HoloLensSurvey.h: HoloLens platform hardware-survey classes
+==============================================================================================*/
+
+#pragma once
+#include "GenericPlatform/GenericPlatformSurvey.h"
+
+/**
+* HoloLens implementation of FGenericPlatformSurvey
+**/
+struct FHoloLensPlatformSurvey : public FGenericPlatformSurvey
+{
+	// default implementation for now
+};
+
+typedef FHoloLensPlatformSurvey FPlatformSurvey;
\ No newline at end of file
diff --git a/Engine/Platforms/Hololens/Source/Runtime/ApplicationCore/Public/HoloLensWindow.h b/Engine/Platforms/Hololens/Source/Runtime/ApplicationCore/Public/HoloLensWindow.h
new file mode 100644
--- /dev/null
+++ b/Engine/Platforms/Hololens/Source/Runtime/ApplicationCore/Public/HoloLensWindow.h
@@ -0,0 +1,49 @@
+// Copyright Epic Games, Inc. All Rights Reserved.
+
+#pragma once
+
+#include "GenericPlatform/GenericWindow.h"
+#include "GenericPlatform/GenericApplication.h"
+#include "Templates/SharedPointer.h"
+
+
+/**
+* A platform specific implementation of FNativeWindow.
+*/
+class CORE_API FHoloLensWindow
+	: public FGenericWindow, public TSharedFromThis<FHoloLensWindow>
+{
+public:
+
+	/** Destructor. */
+	virtual ~FHoloLensWindow();
+
+	/** Create a new FHoloLensWindow. */
+	static TSharedRef<FHoloLensWindow> Make();
+
+	/** Init a FHoloLensWindow. */
+	void Initialize(class FHoloLensApplication* const Application, const TSharedRef<FGenericWindowDefinition>& InDefinition);
+
+	static FPlatformRect GetOSWindowBounds();
+	
+public:
+	virtual void ReshapeWindow(int32 X, int32 Y, int32 Width, int32 Height) override;
+
+	virtual void AdjustCachedSize(FVector2D& Size) const override;
+
+	virtual void SetWindowMode(EWindowMode::Type InNewWindowMode) override;
+
+	virtual EWindowMode::Type GetWindowMode() const override;
+
+public:
+	static int32 ConvertDipsToPixels(int32 Dips, float Dpi);
+	static int32 ConvertPixelsToDips(int32 Pixels, float Dpi);
+
+private:
+
+	/** Protect the constructor; only TSharedRefs of this class can be made. */
+	FHoloLensWindow();
+
+	/** The mode that the window is in (windowed, fullscreen, windowedfullscreen ) */
+	EWindowMode::Type WindowMode;
+};
diff --git a/Engine/Platforms/Hololens/Source/Runtime/Core/Core_HoloLens.build.cs b/Engine/Platforms/Hololens/Source/Runtime/Core/Core_HoloLens.build.cs
new file mode 100644
--- /dev/null
+++ b/Engine/Platforms/Hololens/Source/Runtime/Core/Core_HoloLens.build.cs
@@ -0,0 +1,23 @@
+// Copyright Epic Games, Inc. All Rights Reserved.
+
+namespace UnrealBuildTool.Rules
+{
+	public class Core_HoloLens : Core
+	{
+		public Core_HoloLens(ReadOnlyTargetRules Target) : base(Target)
+		{
+			PublicIncludePaths.Add("../Platforms/HoloLens/Source/Runtime/Core/Public");
+			AddEngineThirdPartyPrivateStaticDependencies(Target,
+				"zlib");
+
+			AddEngineThirdPartyPrivateStaticDependencies(Target,
+				"IntelTBB",
+				"XInput"
+				);
+
+			PublicDefinitions.Add("WITH_VS_PERF_PROFILER=0");
+			PublicDefinitions.Add("IS_RUNNING_GAMETHREAD_ON_EXTERNAL_THREAD=1");
+
+		}
+	}
+}
diff --git a/Engine/Platforms/Hololens/Source/Runtime/Core/Private/HoloLensEvent.h b/Engine/Platforms/Hololens/Source/Runtime/Core/Private/HoloLensEvent.h
new file mode 100644
--- /dev/null
+++ b/Engine/Platforms/Hololens/Source/Runtime/Core/Private/HoloLensEvent.h
@@ -0,0 +1,81 @@
+// Copyright Epic Games, Inc. All Rights Reserved.
+
+/*=============================================================================
+	HoloLensEvent.h: Declares the FEventWin class.
+=============================================================================*/
+
+#pragma once
+
+#include "Microsoft/AllowMicrosoftPlatformTypes.h"
+
+/**
+ * Implements the HoloLens version of the FEvent interface.
+ */
+class FEventHoloLens : public FEvent
+{
+public:
+
+	/**
+	 * Default constructor.
+	 */
+	FEventHoloLens()
+		: Event(NULL)
+	{
+	}
+
+	/**
+	 * Destructor.
+	 */
+	virtual ~FEventHoloLens()
+	{
+		if (Event != NULL)
+		{
+			CloseHandle(Event);
+		}
+	}
+
+
+public:
+
+	virtual bool Create (bool bIsManualReset = false) override
+	{
+		// Create the event and default it to non-signaled
+//		Event = CreateEvent(NULL, bIsManualReset, 0, NULL);
+		Event = CreateEventEx(NULL, NULL, bIsManualReset ? CREATE_EVENT_MANUAL_RESET : 0, EVENT_ALL_ACCESS);
+		ManualReset = bIsManualReset;
+
+		return Event != NULL;
+	}
+
+	virtual bool IsManualReset() override
+	{
+		return ManualReset;
+	}
+
+	virtual void Trigger () override
+	{
+		check(Event);
+
+		SetEvent(Event);
+	}
+
+	virtual void Reset () override
+	{
+		check(Event);
+
+		ResetEvent(Event);
+	}
+
+	virtual bool Wait(uint32 WaitTime, const bool bIgnoreThreadIdleStats = false) override;
+
+
+private:
+
+	/** Whether the signaled state of the event needs to be reset manually. */
+	bool ManualReset;
+
+	// Holds the handle to the event
+	HANDLE Event;
+};
+
+#include "Microsoft/HideMicrosoftPlatformTypes.h"
diff --git a/Engine/Platforms/Hololens/Source/Runtime/Core/Private/HoloLensOutputDevices.cpp b/Engine/Platforms/Hololens/Source/Runtime/Core/Private/HoloLensOutputDevices.cpp
new file mode 100644
--- /dev/null
+++ b/Engine/Platforms/Hololens/Source/Runtime/Core/Private/HoloLensOutputDevices.cpp
@@ -0,0 +1,23 @@
+// Copyright Epic Games, Inc. All Rights Reserved.
+
+/*=============================================================================
+	HoloLensOutputDevices.cpp: HoloLens implementations of OutputDevices functions
+=============================================================================*/
+
+#include "HoloLensPlatformOutputDevices.h"
+#include "Misc/App.h"
+#include "Misc/CoreDelegates.h"
+
+#include "HAL/FeedbackContextAnsi.h"
+#include "HoloLensOutputDevicesPrivate.h"
+
+DEFINE_LOG_CATEGORY(LogHoloLensOutputDevices);
+
+//////////////////////////////////
+// FHoloLensOutputDevices
+//////////////////////////////////
+class FOutputDevice* FHoloLensOutputDevices::GetEventLog()
+{
+	static FOutputDeviceEventLog Singleton;
+	return &Singleton;
+}
\ No newline at end of file
diff --git a/Engine/Platforms/Hololens/Source/Runtime/Core/Private/HoloLensOutputDevicesPrivate.h b/Engine/Platforms/Hololens/Source/Runtime/Core/Private/HoloLensOutputDevicesPrivate.h
new file mode 100644
--- /dev/null
+++ b/Engine/Platforms/Hololens/Source/Runtime/Core/Private/HoloLensOutputDevicesPrivate.h
@@ -0,0 +1,92 @@
+// Copyright Epic Games, Inc. All Rights Reserved.
+
+/*=============================================================================
+	HoloLensOutputDevicesPrivate.h: HoloLens implementations of output devices
+=============================================================================*/
+
+#pragma once
+#include "Misc/OutputDeviceError.h"
+#include "Misc/OutputDeviceConsole.h"
+
+DECLARE_LOG_CATEGORY_EXTERN(LogHoloLensOutputDevices, Log, All);
+
+/**
+* Output device that writes to HoloLens Event Log
+*/
+class FOutputDeviceEventLog :
+	public FOutputDevice
+{
+	Windows::Foundation::Diagnostics::LoggingChannel^ EtwLogChannel;
+
+public:
+	/**
+	* Constructor, initializing member variables
+	*/
+	FOutputDeviceEventLog()
+	{
+#if !NO_LOGGING
+		// Always use the default logging channel GUID here.  This is simpler than the prior approach of sometimes
+		// using the session id, and UAT can still filter on the provider name.
+		static const Platform::Guid MicrosoftWindowsDiagnoticsLoggingChannelId(0x4bd2826e, 0x54a1, 0x4ba9, 0xbf, 0x63, 0x92, 0xb7, 0x3e, 0xa1, 0xac, 0x4a);
+		EtwLogChannel = ref new Windows::Foundation::Diagnostics::LoggingChannel(ref new Platform::String(FApp::GetProjectName()), nullptr, MicrosoftWindowsDiagnoticsLoggingChannelId);
+#endif
+	}
+
+	/** Destructor that cleans up any remaining resources */
+	virtual ~FOutputDeviceEventLog()
+	{
+		TearDown();
+	}
+
+	virtual void Serialize(const TCHAR* Buffer, ELogVerbosity::Type Verbosity, const class FName& Category) override
+	{
+#if !NO_LOGGING
+		if (EtwLogChannel != nullptr)
+		{
+			Windows::Foundation::Diagnostics::LoggingLevel LogWithLevel = GetWindowsLoggingLevelFromUEVerbosity(Verbosity);
+			// we need to use Information level as minimal for getting all information to the editor
+			// and giving user ability to read the whole meaningful device's log
+			if (LogWithLevel >= Windows::Foundation::Diagnostics::LoggingLevel::Information)
+			{
+				EtwLogChannel->LogMessage(ref new Platform::String(Buffer), LogWithLevel);
+			}
+		}
+#endif
+	}
+
+	Windows::Foundation::Diagnostics::LoggingLevel GetWindowsLoggingLevelFromUEVerbosity(ELogVerbosity::Type Verbosity)
+	{
+		switch (Verbosity)
+		{
+		case ELogVerbosity::Error:
+			return Windows::Foundation::Diagnostics::LoggingLevel::Error;
+
+		case ELogVerbosity::Warning:
+			return Windows::Foundation::Diagnostics::LoggingLevel::Warning;
+
+		case ELogVerbosity::Fatal:
+			return Windows::Foundation::Diagnostics::LoggingLevel::Critical;
+
+		case ELogVerbosity::Display:
+		case ELogVerbosity::Log:
+			return Windows::Foundation::Diagnostics::LoggingLevel::Information;
+
+		case ELogVerbosity::Verbose:
+		case ELogVerbosity::VeryVerbose:
+		default:
+			return Windows::Foundation::Diagnostics::LoggingLevel::Verbose;
+		}
+	}
+
+	/** Does nothing */
+	virtual void Flush(void)
+	{
+	}
+
+	/**
+	* Closes any event log handles that are open
+	*/
+	virtual void TearDown(void)
+	{
+	}
+};
\ No newline at end of file
diff --git a/Engine/Platforms/Hololens/Source/Runtime/Core/Private/HoloLensPlatformAtomics.cpp b/Engine/Platforms/Hololens/Source/Runtime/Core/Private/HoloLensPlatformAtomics.cpp
new file mode 100644
--- /dev/null
+++ b/Engine/Platforms/Hololens/Source/Runtime/Core/Private/HoloLensPlatformAtomics.cpp
@@ -0,0 +1,21 @@
+// Copyright Epic Games, Inc. All Rights Reserved.
+
+#include "HoloLensPlatformAtomics.h"
+#include "CoreMinimal.h"
+
+#if PLATFORM_HOLOLENS
+
+void FHoloLensAtomics::HandleAtomicsFailure( const TCHAR* InFormat, ... )
+{	
+	TCHAR TempStr[1024];
+	va_list Ptr;
+
+	va_start( Ptr, InFormat );	
+	FPlatformString::GetVarArgs( TempStr, UE_ARRAY_COUNT(TempStr), InFormat, Ptr );
+	va_end( Ptr );
+
+	UE_LOG(LogTemp, Log, TempStr);
+	check( 0 );
+}
+
+#endif //PLATFORM_HOLOLENS
\ No newline at end of file
