move HoloLens runtime source to platform extension folder.

diff --git a/Engine/Platforms/Hololens/Source/Runtime/Core/Private/HoloLensPlatformCrashContext.cpp b/Engine/Platforms/Hololens/Source/Runtime/Core/Private/HoloLensPlatformCrashContext.cpp
new file mode 100644
--- /dev/null
+++ b/Engine/Platforms/Hololens/Source/Runtime/Core/Private/HoloLensPlatformCrashContext.cpp
@@ -0,0 +1,510 @@
+// Copyright Epic Games, Inc. All Rights Reserved.
+
+
+#include "HoloLensPlatformCrashContext.h"
+
+#if WER_CUSTOM_REPORTS
+
+#include "HAL/PlatformMallocCrash.h"
+#include "HAL/ExceptionHandling.h"
+#include "HoloLensPlatformCrashContext.h"
+#include "../../Launch/Resources/Version.h"
+
+#include "Microsoft/AllowMicrosoftPlatformTypes.h"
+
+	#include <strsafe.h>
+	#include <werapi.h>
+	#pragma comment( lib, "wer.lib" )
+	#include <dbghelp.h>
+	#include <Shlwapi.h>
+
+#pragma comment( lib, "version.lib" )
+#pragma comment( lib, "Shlwapi.lib" )
+
+
+namespace
+{
+static int32 ReportCrashCallCount = 0;
+
+/**
+ * Write a Windows minidump to disk
+ * @param Path Full path of file to write (normally a .dmp file)
+ * @param ExceptionInfo Pointer to structure containing the exception information
+ * @return Success or failure
+ */
+
+// @TODO yrx 2014-10-08 Move to FHoloLensPlatformCrashContext
+bool WriteMinidump(const TCHAR* Path, LPEXCEPTION_POINTERS ExceptionInfo)
+{
+	// Try to create file for minidump.
+	HANDLE FileHandle = CreateFile2(Path, GENERIC_WRITE, 0, CREATE_ALWAYS, nullptr);
+	
+	if (FileHandle == INVALID_HANDLE_VALUE)
+	{
+		return false;
+	}
+
+	// Initialise structure required by MiniDumpWriteDumps
+	MINIDUMP_EXCEPTION_INFORMATION DumpExceptionInfo = {};
+
+	DumpExceptionInfo.ThreadId			= GetCurrentThreadId();
+	DumpExceptionInfo.ExceptionPointers	= ExceptionInfo;
+	DumpExceptionInfo.ClientPointers	= FALSE;
+
+	// CrashContext.runtime-xml is now a part of the minidump file.
+	FHoloLensPlatformCrashContext CrashContext;
+	CrashContext.SerializeContentToBuffer();
+	
+	MINIDUMP_USER_STREAM CrashContextStream ={0};
+	CrashContextStream.Type = FHoloLensPlatformCrashContext::UE4_MINIDUMP_CRASHCONTEXT;
+	CrashContextStream.BufferSize = CrashContext.GetBuffer().GetAllocatedSize();
+	CrashContextStream.Buffer = (void*)*CrashContext.GetBuffer();
+
+	MINIDUMP_USER_STREAM_INFORMATION CrashContextStreamInformation = {0};
+	CrashContextStreamInformation.UserStreamCount = 1;
+	CrashContextStreamInformation.UserStreamArray = &CrashContextStream;
+
+	const BOOL Result = MiniDumpWriteDump(GetCurrentProcess(), GetCurrentProcessId(), FileHandle, MiniDumpNormal, &DumpExceptionInfo, &CrashContextStreamInformation, NULL);
+	CloseHandle(FileHandle);
+
+	return Result == TRUE;
+}
+
+/** 
+ * Get one line of text to describe the crash
+ */
+void GetCrashDescription(WER_REPORT_INFORMATION& ReportInformation)
+{
+	StringCchCopy( ReportInformation.wzDescription, ARRAYSIZE( ReportInformation.wzDescription ), TEXT( "The application crashed while running " ) );
+
+	const TCHAR* Description = IsRunningCommandlet() ? TEXT( "a commandlet" ) :
+		GIsEditor ? TEXT( "the editor" ) :
+		IsRunningDedicatedServer() ? TEXT( "a server" ) :
+		TEXT( "the game" );
+
+	StringCchCat( ReportInformation.wzDescription, ARRAYSIZE( ReportInformation.wzDescription ), Description );
+}
+
+//--------------------
+// @HEMI : GetFileVersionInfo apis aren't part of the "APP" api set, but if a developer specifically changes
+//                   to a non-shippable API set for compilation, we'll still collect and attach the additional module info
+//-----
+#if (WINAPI_FAMILY==WINAPI_FAMILY_ONECORE_APP)
+
+/**
+ * Get the 4 element version number of the module
+ */
+void GetModuleVersion( TCHAR* ModuleName, TCHAR* StringBuffer, DWORD MaxSize )
+{
+	StringCchCopy( StringBuffer, MaxSize, TEXT( "0.0.0.0" ) );
+	
+	DWORD Handle = 0;
+	DWORD InfoSize = GetFileVersionInfoSize( ModuleName, &Handle );
+	if( InfoSize > 0 )
+	{
+		TArray<char> VersionInfo;
+		VersionInfo.SetNum(InfoSize);
+
+		if( GetFileVersionInfo( ModuleName, 0, InfoSize, VersionInfo.GetData() ) )
+		{
+			VS_FIXEDFILEINFO* FixedFileInfo;
+
+			UINT InfoLength = 0;
+			if( VerQueryValue( VersionInfo.GetData(), TEXT( "\\" ), ( void** )&FixedFileInfo, &InfoLength ) )
+			{
+				StringCchPrintf( StringBuffer, MaxSize, TEXT( "%u.%u.%u.%u" ), 
+					HIWORD( FixedFileInfo->dwProductVersionMS ), LOWORD( FixedFileInfo->dwProductVersionMS ), HIWORD( FixedFileInfo->dwProductVersionLS ), LOWORD( FixedFileInfo->dwProductVersionLS ) );
+			}
+		}
+	}
+}
+
+#endif
+//-----
+// @HEMI : end
+//--------------------
+
+/** 
+ * Set the ten Windows Error Reporting parameters
+ *
+ * Parameters 0 through 7 are predefined for Windows
+ * Parameters 8 and 9 are user defined
+ */
+void SetReportParameters( HREPORT ReportHandle, EXCEPTION_POINTERS* ExceptionInfo, const TCHAR* ErrorMessage )
+{
+	HRESULT Result;
+	TCHAR StringBuffer[MAX_SPRINTF] = {0};
+	TCHAR LocalBuffer[MAX_SPRINTF] = {0};
+
+	// Set the parameters for the standard problem signature
+	HMODULE ModuleHandle = GetModuleHandle( NULL );
+
+	StringCchPrintf( StringBuffer, MAX_SPRINTF, TEXT( "UE4-%s" ), FApp::GetGameName() );
+	Result = WerReportSetParameter( ReportHandle, WER_P0, TEXT( "Application Name" ), StringBuffer );
+
+	StringCchPrintf(StringBuffer, MAX_SPRINTF, TEXT("%08x"), GetTimestampForLoadedLibrary(ModuleHandle));
+	Result = WerReportSetParameter(ReportHandle, WER_P2, TEXT("Application Timestamp"), StringBuffer);
+
+//--------------------
+// @HEMI : GetFileVersionInfo apis aren't part of the "APP" api set, but if a developer specifically changes
+//                   to a non-shippable API set for compilation, we'll still collect and attach the additional module info
+//-----
+#if (WINAPI_FAMILY==WINAPI_FAMILY_ONECORE_APP)
+	GetModuleFileName( ModuleHandle, LocalBuffer, MAX_SPRINTF );
+	PathStripPath( LocalBuffer );
+	GetModuleVersion( LocalBuffer, StringBuffer, MAX_SPRINTF );
+	Result = WerReportSetParameter( ReportHandle, WER_P1, TEXT( "Application Version" ), StringBuffer );
+
+	HMODULE FaultModuleHandle = NULL;
+	GetModuleHandleEx( GET_MODULE_HANDLE_EX_FLAG_FROM_ADDRESS | GET_MODULE_HANDLE_EX_FLAG_UNCHANGED_REFCOUNT, ( LPCTSTR )ExceptionInfo->ExceptionRecord->ExceptionAddress, &FaultModuleHandle );
+
+	GetModuleFileName( FaultModuleHandle, LocalBuffer, MAX_SPRINTF );
+	PathStripPath( LocalBuffer );
+	Result = WerReportSetParameter( ReportHandle, WER_P3, TEXT( "Fault Module Name" ), LocalBuffer );
+
+	GetModuleVersion( LocalBuffer, StringBuffer, MAX_SPRINTF );
+	Result = WerReportSetParameter( ReportHandle, WER_P4, TEXT( "Fault Module Version" ), StringBuffer );
+
+	StringCchPrintf(StringBuffer, MAX_SPRINTF, TEXT("%08x"), GetTimestampForLoadedLibrary(FaultModuleHandle));
+	Result = WerReportSetParameter(ReportHandle, WER_P5, TEXT("Fault Module Timestamp"), StringBuffer);
+
+	INT_PTR ExceptionOffset = (char*)(ExceptionInfo->ExceptionRecord->ExceptionAddress) - (char*)FaultModuleHandle;
+	CA_SUPPRESS(6066) // The format specifier should probably be something like %tX, but VS 2013 doesn't support 't'.
+		StringCchPrintf(StringBuffer, MAX_SPRINTF, TEXT("%p"), ExceptionOffset);
+	Result = WerReportSetParameter(ReportHandle, WER_P7, TEXT("Exception Offset"), StringBuffer);
+#endif
+//-----
+// @HEMI : end
+//--------------------
+
+	StringCchPrintf( StringBuffer, MAX_SPRINTF, TEXT( "%08x" ), ExceptionInfo->ExceptionRecord->ExceptionCode );
+	Result = WerReportSetParameter( ReportHandle, WER_P6, TEXT( "Exception Code" ), StringBuffer );
+
+	// Use LocalBuffer to store the error message.
+	FCString::Strncpy( LocalBuffer, ErrorMessage, MAX_SPRINTF );
+
+	// Replace " with ' and replace \n with #
+	for (TCHAR& Char: LocalBuffer)
+	{
+		if (Char == 0)
+		{
+			break;
+		}
+
+		switch (Char)
+		{
+			default: break;
+			case '"':	Char = '\'';	break;
+			case '\r':	Char = '#';		break;
+			case '\n':	Char = '#';		break;
+		}
+	}
+
+	// AssertLog should be ErrorMessage, but this require crash server changes, so don't change this.
+	StringCchPrintf( StringBuffer, MAX_SPRINTF, TEXT( "!%s!AssertLog=\"%s\"" ), FCommandLine::Get(), LocalBuffer );
+	Result = WerReportSetParameter( ReportHandle, WER_P8, TEXT( "Commandline" ), StringBuffer );
+
+	StringCchPrintf( StringBuffer, MAX_SPRINTF, TEXT( "%s!%s!%s!%d" ), TEXT( BRANCH_NAME ), FPlatformProcess::BaseDir(), FPlatformMisc::GetEngineMode(), BUILT_FROM_CHANGELIST );
+	Result = WerReportSetParameter( ReportHandle, WER_P9, TEXT( "BranchBaseDir" ), StringBuffer );
+}
+
+/**
+ * Add a minidump to the Windows Error Report
+ *
+ * Note this has to be a minidump (and not a microdump) for the dbgeng functions to work correctly
+ */
+void AddMiniDump(HREPORT ReportHandle, EXCEPTION_POINTERS* ExceptionInfo)
+{
+	// Add GetMinidumpFile
+	const FString MinidumpFileName = FString::Printf(TEXT("%sDump%d.dmp"), *FPaths::GameLogDir(), FDateTime::UtcNow().GetTicks());
+	
+	if (WriteMinidump(*MinidumpFileName, ExceptionInfo))
+	{
+		WerReportAddFile(ReportHandle, *MinidumpFileName, WerFileTypeMinidump, WER_FILE_ANONYMOUS_DATA);
+	}
+}
+
+/** 
+ * Add miscellaneous files to the report. Currently the log and the video file
+ */
+void AddMiscFiles( HREPORT ReportHandle )
+{
+	FString LogFileName = FPaths::GameLogDir() / FApp::GetGameName() + TEXT( ".log" );
+	WerReportAddFile( ReportHandle, *LogFileName, WerFileTypeOther, WER_FILE_ANONYMOUS_DATA ); 
+
+	FString CrashVideoPath = FPaths::GameLogDir() / TEXT( "CrashVideo.avi" );
+	WerReportAddFile( ReportHandle, *CrashVideoPath, WerFileTypeOther, WER_FILE_ANONYMOUS_DATA );
+}
+
+/**
+ * Enum indicating whether to run the crash reporter UI
+ */
+enum class EErrorReportUI
+{
+	/** Ask the user for a description */
+	ShowDialog,
+
+	/** Silently uploaded the report */
+	ReportInUnattendedMode	
+};
+
+/** 
+ * Create a Windows Error Report, add the user log and video, and add it to the WER queue
+ */
+int32 ReportCrashUsingWindowsErrorReporting(EXCEPTION_POINTERS* ExceptionInfo, const TCHAR* ErrorMessage, EErrorReportUI ReportUI)
+{
+	// Flush out the log
+	GLog->Panic();
+
+	// Construct the report details
+	WER_REPORT_INFORMATION ReportInformation = {sizeof( WER_REPORT_INFORMATION )};
+
+	StringCchCopy( ReportInformation.wzConsentKey, ARRAYSIZE( ReportInformation.wzConsentKey ), TEXT( "" ) );
+
+	StringCchCopy( ReportInformation.wzApplicationName, ARRAYSIZE( ReportInformation.wzApplicationName ), TEXT("UE4-") );
+	StringCchCat( ReportInformation.wzApplicationName, ARRAYSIZE( ReportInformation.wzApplicationName ), FApp::GetGameName() );
+
+	StringCchCopy( ReportInformation.wzApplicationPath, ARRAYSIZE( ReportInformation.wzApplicationPath ), FPlatformProcess::BaseDir() );
+	StringCchCat( ReportInformation.wzApplicationPath, ARRAYSIZE( ReportInformation.wzApplicationPath ), FPlatformProcess::ExecutableName() );
+	StringCchCat( ReportInformation.wzApplicationPath, ARRAYSIZE( ReportInformation.wzApplicationPath ), TEXT( ".exe" ) );
+
+	GetCrashDescription( ReportInformation );
+
+	// Create a crash event report
+	HREPORT ReportHandle = NULL;
+	if( WerReportCreate( APPCRASH_EVENT, WerReportApplicationCrash, &ReportInformation, &ReportHandle ) == S_OK )
+	{
+		// Set the standard set of a crash parameters
+		SetReportParameters( ReportHandle, ExceptionInfo, ErrorMessage );
+
+		// Add a manually generated minidump
+		AddMiniDump( ReportHandle, ExceptionInfo );
+
+		// Add the log and video
+		AddMiscFiles( ReportHandle );
+
+		// Submit
+		WER_SUBMIT_RESULT SubmitResult;
+		WerReportSubmit( ReportHandle, WerConsentAlwaysPrompt, WER_SUBMIT_QUEUE | WER_SUBMIT_BYPASS_DATA_THROTTLING, &SubmitResult );
+
+		// Cleanup
+		WerReportCloseHandle( ReportHandle );
+	}
+
+	// Let the system take back over (return value only used by NewReportEnsure)
+	return EXCEPTION_CONTINUE_EXECUTION;
+}
+
+
+} // end anonymous namespace
+#include "Microsoft/HideMicrosoftPlatformTypes.h"
+
+#endif
+
+#include "Microsoft/AllowMicrosoftPlatformTypes.h"
+
+/** 
+ * Report an ensure to the crash reporting system
+ */
+void NewReportEnsure( const TCHAR* ErrorMessage )
+{
+#if !PLATFORM_SEH_EXCEPTIONS_DISABLED && !PLATFORM_HOLOLENS
+	__try
+#endif
+	{
+		FPlatformMisc::RaiseException( 1 );
+	}
+#if !PLATFORM_SEH_EXCEPTIONS_DISABLED && !PLATFORM_HOLOLENS
+	__except(ReportCrashUsingWindowsErrorReporting(GetExceptionInformation(), ErrorMessage, EErrorReportUI::ReportInUnattendedMode))
+	CA_SUPPRESS(6322)
+	{
+	}
+#endif
+}
+
+
+#include "Microsoft/HideMicrosoftPlatformTypes.h"
+#if WER_CUSTOM_REPORTS
+
+// Original code below
+
+#include "Microsoft/AllowMicrosoftPlatformTypes.h"
+	#include <ErrorRep.h>
+	#include <DbgHelp.h>
+#include "Microsoft/HideMicrosoftPlatformTypes.h"
+
+#pragma comment(lib, "Faultrep.lib")
+
+/** 
+ * Creates an info string describing the given exception record.
+ * See MSDN docs on EXCEPTION_RECORD.
+ */
+#include "Microsoft/AllowMicrosoftPlatformTypes.h"
+void CreateExceptionInfoString(EXCEPTION_RECORD* ExceptionRecord)
+{
+	// @TODO yrx 2014-08-18 Fix FString usage?
+	FString ErrorString = TEXT("Unhandled Exception: ");
+
+#define HANDLE_CASE(x) case x: ErrorString += TEXT(#x); break;
+
+	switch (ExceptionRecord->ExceptionCode)
+	{
+	case EXCEPTION_ACCESS_VIOLATION:
+		ErrorString += TEXT("EXCEPTION_ACCESS_VIOLATION ");
+		if (ExceptionRecord->ExceptionInformation[0] == 0)
+		{
+			ErrorString += TEXT("reading address ");
+		}
+		else if (ExceptionRecord->ExceptionInformation[0] == 1)
+		{
+			ErrorString += TEXT("writing address ");
+		}
+		ErrorString += FString::Printf(TEXT("0x%016llx"), ExceptionRecord->ExceptionInformation[1]);
+		break;
+	HANDLE_CASE(EXCEPTION_ARRAY_BOUNDS_EXCEEDED)
+	HANDLE_CASE(EXCEPTION_DATATYPE_MISALIGNMENT)
+	HANDLE_CASE(EXCEPTION_FLT_DENORMAL_OPERAND)
+	HANDLE_CASE(EXCEPTION_FLT_DIVIDE_BY_ZERO)
+	HANDLE_CASE(EXCEPTION_FLT_INVALID_OPERATION)
+	HANDLE_CASE(EXCEPTION_ILLEGAL_INSTRUCTION)
+	HANDLE_CASE(EXCEPTION_INT_DIVIDE_BY_ZERO)
+	HANDLE_CASE(EXCEPTION_PRIV_INSTRUCTION)
+	HANDLE_CASE(EXCEPTION_STACK_OVERFLOW)
+	default:
+		ErrorString += FString::Printf(TEXT("0x%08x"), (uint32)ExceptionRecord->ExceptionCode);
+	}
+
+#if WITH_EDITORONLY_DATA
+	FCString::Strncpy(GErrorExceptionDescription, *ErrorString, UE_ARRAY_COUNT(GErrorExceptionDescription));
+#endif
+#undef HANDLE_CASE
+}
+#include "Microsoft/HideMicrosoftPlatformTypes.h"
+
+int32 ReportCrash( LPEXCEPTION_POINTERS ExceptionInfo )
+{
+	// Switch to malloc crash.
+	FMallocCrash::Get().SetAsGMalloc();
+
+	// Only create a minidump the first time this function is called.
+	// (Can be called the first time from the RenderThread, then a second time from the MainThread.)
+	if( FPlatformAtomics::InterlockedIncrement( &ReportCrashCallCount ) != 1 )
+	{
+		return EXCEPTION_EXECUTE_HANDLER;
+	}
+
+	FCoreDelegates::OnHandleSystemError.Broadcast();
+
+	// Note: ReportCrashUsingWindowsErrorReporting reads the callstack from GErrorHist
+	const SIZE_T StackTraceSize = 65535;
+	ANSICHAR* StackTrace = (ANSICHAR*) GMalloc->Malloc( StackTraceSize );
+	StackTrace[0] = 0;
+	// Walk the stack and dump it to the allocated memory.
+	FPlatformStackWalk::StackWalkAndDump( StackTrace, StackTraceSize, 0, ExceptionInfo->ContextRecord );
+
+	// @TODO yrx 2014-08-20 Make this constant.
+	if( ExceptionInfo->ExceptionRecord->ExceptionCode != 1 )
+	{
+		CreateExceptionInfoString( ExceptionInfo->ExceptionRecord );
+		FCString::Strncat( GErrorHist, GErrorExceptionDescription, UE_ARRAY_COUNT( GErrorHist ) );
+		FCString::Strncat( GErrorHist, TEXT( "\r\n\r\n" ), UE_ARRAY_COUNT( GErrorHist ) );
+	}
+
+	FCString::Strncat( GErrorHist, ANSI_TO_TCHAR(StackTrace), UE_ARRAY_COUNT(GErrorHist) );
+
+	GMalloc->Free( StackTrace );
+
+	ReportCrashUsingWindowsErrorReporting(ExceptionInfo, GErrorMessage, EErrorReportUI::ShowDialog);
+	WriteMinidump(MiniDumpFilenameW, ExceptionInfo);
+
+#if UE_BUILD_SHIPPING && WITH_EDITOR
+	uint32 dwOpt = 0;
+	EFaultRepRetVal repret = ReportFault( ExceptionInfo, dwOpt);
+#endif
+
+	return EXCEPTION_EXECUTE_HANDLER;
+}
+
+#else
+int32 ReportCrash(LPEXCEPTION_POINTERS ExceptionInfo)
+{
+	throw *ExceptionInfo;
+}
+#endif
+
+#if WER_CUSTOM_REPORTS
+void ReportHang(const TCHAR* ErrorMessage, const uint64* StackFrames, int32 NumStackFrames, uint32 HungThreadId)
+{
+	if (ReportCrashCallCount > 0 || FDebug::HasAsserted())
+	{
+		// Don't report ensures after we've crashed/asserted, they simply may be a result of the crash as
+		// the engine is already in a bad state.
+		return;
+	}
+
+	FHoloLensPlatformCrashContext CrashContext(ECrashContextType::Ensure, ErrorMessage);
+	CrashContext.SetPortableCallStack(StackFrames, NumStackFrames);
+	CrashContext.SetCrashedThreadId(HungThreadId);
+	CrashContext.CaptureAllThreadContexts();
+
+	EErrorReportUI ReportUI = IsInteractiveEnsureMode() ? EErrorReportUI::ShowDialog : EErrorReportUI::ReportInUnattendedMode;
+	ReportCrashUsingCrashReportClient(CrashContext, nullptr, ReportUI);
+}
+#else
+void ReportHang(const TCHAR* ErrorMessage, const uint64* StackFrames, int32 NumStackFrames, uint32 HungThreadId)
+{
+
+}
+#endif
+
+/** Implement platform specific static cleanup function */
+void FGenericCrashContext::CleanupPlatformSpecificFiles()
+{
+#if WER_CUSTOM_REPORTS
+	FString CrashVideoPath = FPaths::ProjectLogDir() / TEXT("CrashVideo.avi");
+	IFileManager::Get().Delete(*CrashVideoPath);
+#endif
+}
+
+#if WER_CUSTOM_REPORTS
+
+static FCriticalSection EnsureLock;
+static bool bReentranceGuard = false;
+
+void ReportEnsure(const TCHAR* ErrorMessage, int NumStackFramesToIgnore)
+{
+	if (ReportCrashCallCount > 0 || FDebug::HasAsserted())
+	{
+		// Don't report ensures after we've crashed/asserted, they simply may be a result of the crash as
+		// the engine is already in a bad state.
+		return;
+	}
+
+	// Simple re-entrance guard.
+	EnsureLock.Lock();
+
+	if (bReentranceGuard)
+	{
+		EnsureLock.Unlock();
+		return;
+	}
+
+	// Stop checking heartbeat for this thread (and stop the gamethread hitch detector if we're the game thread).
+	// Ensure can take a lot of time (when stackwalking), so we don't want hitches/hangs firing.
+	// These are no-ops on threads that didn't already have a heartbeat etc.
+	FSlowHeartBeatScope SuspendHeartBeat(true);
+	FDisableHitchDetectorScope SuspendGameThreadHitch;
+
+	bReentranceGuard = true;
+
+	NewReportEnsure(ErrorMessage);
+
+	bReentranceGuard = false;
+	EnsureLock.Unlock();
+}
+#else
+void ReportEnsure(const TCHAR* ErrorMessage, int NumStackFramesToIgnore)
+{
+
+}
+#endif
diff --git a/Engine/Platforms/Hololens/Source/Runtime/Core/Private/HoloLensRunnableThread.cpp b/Engine/Platforms/Hololens/Source/Runtime/Core/Private/HoloLensRunnableThread.cpp
new file mode 100644
--- /dev/null
+++ b/Engine/Platforms/Hololens/Source/Runtime/Core/Private/HoloLensRunnableThread.cpp
@@ -0,0 +1,100 @@
+// Copyright Epic Games, Inc. All Rights Reserved.
+
+#include "HoloLensRunnableThread.h"
+#include "Misc/OutputDeviceError.h"
+#include "HAL/ExceptionHandling.h"
+#include "Misc/FeedbackContext.h"
+#include "CoreGlobals.h"
+
+DEFINE_LOG_CATEGORY_STATIC(LogThreadingWindows, Log, All);
+
+
+uint32 FRunnableThreadHoloLens::GuardedRun()
+{
+	uint32 ExitCode = 1;
+
+	FPlatformProcess::SetThreadAffinityMask(ThreadAffinityMask);
+
+#if !PLATFORM_SEH_EXCEPTIONS_DISABLED
+	if (!FPlatformMisc::IsDebuggerPresent() || GAlwaysReportCrash)
+	{
+		__try
+		{
+			ExitCode = Run();
+		}
+		__except (ReportCrash(GetExceptionInformation()))
+		{
+			// Make sure the information which thread crashed makes it into the log.
+			UE_LOG(LogThreadingWindows, Error, TEXT("Runnable thread %s crashed."), *ThreadName);
+			GWarn->Flush();
+
+			// Append the thread name at the end of the error report.
+			FCString::Strncat(GErrorHist, LINE_TERMINATOR TEXT("Crash in runnable thread "), UE_ARRAY_COUNT(GErrorHist));
+			FCString::Strncat(GErrorHist, *ThreadName, UE_ARRAY_COUNT(GErrorHist));
+
+			ExitCode = 1;
+			// Generate status report.				
+			GError->HandleError();
+			// Throw an error so that main thread shuts down too (otherwise task graph stalls forever).
+			UE_LOG(LogThreadingWindows, Fatal, TEXT("Runnable thread %s crashed."), *ThreadName);
+		}
+	}
+	else
+#endif
+	{
+		ExitCode = Run();
+	}
+
+	return ExitCode;
+}
+
+
+uint32 FRunnableThreadHoloLens::Run()
+{
+	// Assume we'll fail init
+	uint32 ExitCode = 1;
+	check(Runnable);
+
+	// Initialize the runnable object
+	if (Runnable->Init() == true)
+	{
+		// Initialization has completed, release the sync event
+		ThreadInitSyncEvent->Trigger();
+
+		SetTls();
+
+		// Now run the task that needs to be done
+		ExitCode = Runnable->Run();
+		// Allow any allocated resources to be cleaned up
+		Runnable->Exit();
+	}
+	else
+	{
+		// Initialization has failed, release the sync event
+		ThreadInitSyncEvent->Trigger();
+	}
+
+	return ExitCode;
+}
+
+int FRunnableThreadHoloLens::TranslateThreadPriority(EThreadPriority Priority)
+{
+	// If this triggers, 
+	static_assert(TPri_Num == 7, "Need to add a case for new TPri_xxx enum value");
+
+	switch (Priority)
+	{
+	case TPri_AboveNormal:			return THREAD_PRIORITY_ABOVE_NORMAL;
+	case TPri_Normal:				return THREAD_PRIORITY_NORMAL;
+	case TPri_BelowNormal:			return THREAD_PRIORITY_BELOW_NORMAL;
+	case TPri_Highest:				return THREAD_PRIORITY_HIGHEST;
+	case TPri_TimeCritical:			return THREAD_PRIORITY_HIGHEST;
+	case TPri_Lowest:				return THREAD_PRIORITY_LOWEST;
+
+		// There is no such things as slightly below normal on Windows.
+		// This can't be below normal since we don't want latency sensitive task to go to efficient cores on Alder Lake.
+	case TPri_SlightlyBelowNormal:	return THREAD_PRIORITY_NORMAL;
+
+	default: UE_LOG(LogHAL, Fatal, TEXT("Unknown Priority passed to TranslateThreadPriority()")); return THREAD_PRIORITY_NORMAL;
+	}
+}
diff --git a/Engine/Platforms/Hololens/Source/Runtime/Core/Private/HoloLensRunnableThread.h b/Engine/Platforms/Hololens/Source/Runtime/Core/Private/HoloLensRunnableThread.h
new file mode 100644
--- /dev/null
+++ b/Engine/Platforms/Hololens/Source/Runtime/Core/Private/HoloLensRunnableThread.h
@@ -0,0 +1,202 @@
+// Copyright Epic Games, Inc. All Rights Reserved.
+
+#pragma once
+
+#include "CoreMinimal.h"
+#include "HAL/RunnableThread.h"
+#include "HAL/Runnable.h"
+#include "HAL/Event.h"
+#include "HAL/ThreadManager.h"
+#include "Containers/StringConv.h"
+#include "Microsoft/WindowsHWrapper.h"
+#include "HoloLensPlatformProcess.h"
+#include "Microsoft/AllowMicrosoftPlatformTypes.h"
+
+/**
+* This class has a lot of code that is very similar to FRunnableThreadWin.  It would make a lot of sense to refactor into an FRunnableThreadMicrosoft 
+* from which both this and FRunnableThreadWin could derive.  Unfortunately another platform is using the FRunnableThreadWin .h and providing its own
+* function definitions in a different .cpp and there are dire perforance, debugging, crash reporting, and exception handling concerns about the use of virtual
+* functions in this thread class.  So this refactor would have to be done carefully on all platforms.
+*/
+class FRunnableThreadHoloLens
+	: public FRunnableThread
+{
+	/** The thread handle for the thread. */
+	HANDLE Thread = 0;
+
+	/**
+	* Helper function to set thread names, visible by the debugger.
+	* @param ThreadID		Thread ID for the thread who's name is going to be set
+	* @param ThreadName	Name to set
+	*/
+	static void SetThreadName(uint32 ThreadID, LPCSTR ThreadName)
+	{
+#if !PLATFORM_SEH_EXCEPTIONS_DISABLED
+		/**
+		* Code setting the thread name for use in the debugger.
+		*
+		* http://msdn.microsoft.com/en-us/library/xcb2z8hs.aspx
+		*/
+		const uint32 MS_VC_EXCEPTION = 0x406D1388;
+
+		struct THREADNAME_INFO
+		{
+			uint32 dwType;		// Must be 0x1000.
+			LPCSTR szName;		// Pointer to name (in user addr space).
+			uint32 dwThreadID;	// Thread ID (-1=caller thread).
+			uint32 dwFlags;		// Reserved for future use, must be zero.
+		};
+
+		// on the xbox setting thread names messes up the XDK COM API that UnrealConsole uses so check to see if they have been
+		// explicitly enabled
+		Sleep(10);
+		THREADNAME_INFO ThreadNameInfo;
+		ThreadNameInfo.dwType = 0x1000;
+		ThreadNameInfo.szName = ThreadName;
+		ThreadNameInfo.dwThreadID = ThreadID;
+		ThreadNameInfo.dwFlags = 0;
+
+		__try
+		{
+			RaiseException(MS_VC_EXCEPTION, 0, sizeof(ThreadNameInfo) / sizeof(ULONG_PTR), (ULONG_PTR*)&ThreadNameInfo);
+		}
+		__except (EXCEPTION_EXECUTE_HANDLER)
+			CA_SUPPRESS(6322)
+		{
+		}
+#endif
+	}
+
+	/**
+	* The thread entry point. Simply forwards the call on to the right
+	* thread main function
+	*/
+	static DWORD STDCALL _ThreadProc(LPVOID pThis)
+	{
+		check(pThis);
+		auto* ThisThread = (FRunnableThreadHoloLens*)pThis;
+		FThreadManager::Get().AddThread(ThisThread->GetThreadID(), ThisThread);
+		return ThisThread->GuardedRun();
+	}
+
+	/** Guarding works only if debugger is not attached or GAlwaysReportCrash is true. */
+	uint32 GuardedRun();
+
+	/**
+	* The real thread entry point. It calls the Init/Run/Exit methods on
+	* the runnable object
+	*/
+	uint32 Run();
+
+public:
+	~FRunnableThreadHoloLens()
+	{
+		// Clean up our thread if it is still active
+		if (Thread != NULL)
+		{
+			Kill(true);
+		}
+	}
+
+	virtual void SetThreadPriority(EThreadPriority NewPriority) override
+	{
+		// Don't bother calling the OS if there is no need
+		if (NewPriority != ThreadPriority)
+		{
+			ThreadPriority = NewPriority;
+			// Change the priority on the thread
+			::SetThreadPriority(Thread, TranslateThreadPriority(ThreadPriority));
+		}
+	}
+
+	virtual void Suspend(bool bShouldPause = 1) override
+	{
+		check(Thread);
+		if (bShouldPause == true)
+		{
+			SuspendThread(Thread);
+		}
+		else
+		{
+			ResumeThread(Thread);
+		}
+	}
+
+	virtual bool Kill(bool bShouldWait = false) override
+	{
+		check(Thread && "Did you forget to call Create()?");
+		bool bDidExitOK = true;
+		// Let the runnable have a chance to stop without brute force killing
+		if (Runnable)
+		{
+			Runnable->Stop();
+		}
+		// If waiting was specified, wait the amount of time. If that fails,
+		// brute force kill that thread. Very bad as that might leak.
+		if (bShouldWait == true)
+		{
+			// Wait indefinitely for the thread to finish.  IMPORTANT:  It's not safe to just go and
+			// kill the thread with TerminateThread() as it could have a mutex lock that's shared
+			// with a thread that's continuing to run, which would cause that other thread to
+			// dead-lock.  (This can manifest itself in code as simple as the synchronization
+			// object that is used by our logging output classes.  Trust us, we've seen it!)
+			WaitForSingleObjectEx(Thread, INFINITE, FALSE);
+		}
+		// Now clean up the thread handle so we don't leak
+		CloseHandle(Thread);
+		Thread = NULL;
+
+		return bDidExitOK;
+	}
+
+	virtual void WaitForCompletion() override
+	{
+		// Block until this thread exits
+		WaitForSingleObjectEx(Thread, INFINITE, FALSE);
+	}
+
+	static int TranslateThreadPriority(EThreadPriority Priority);
+protected:
+
+	virtual bool CreateInternal(FRunnable* InRunnable, const TCHAR* InThreadName,
+		uint32 InStackSize = 0,
+		EThreadPriority InThreadPri = TPri_Normal, uint64 InThreadAffinityMask = 0,
+		EThreadCreateFlags InCreateFlags = EThreadCreateFlags::None) override
+	{
+		check(InRunnable);
+		Runnable = InRunnable;
+		ThreadAffinityMask = InThreadAffinityMask;
+
+		// Create a sync event to guarantee the Init() function is called first
+		ThreadInitSyncEvent = FPlatformProcess::GetSynchEventFromPool(true);
+
+		ThreadName = InThreadName ? InThreadName : TEXT("Unnamed UE");
+		ThreadPriority = InThreadPri;
+
+		// Create the thread as suspended, so we can ensure ThreadId is initialized and the thread manager knows about the thread before it runs.
+		Thread = CreateThread(NULL, InStackSize, _ThreadProc, this, STACK_SIZE_PARAM_IS_A_RESERVATION | CREATE_SUSPENDED, (DWORD *)&ThreadID);
+
+		// If it fails, clear all the vars
+		if (Thread == NULL)
+		{
+			Runnable = NULL;
+		}
+		else
+		{
+			ResumeThread(Thread);
+
+			// Let the thread start up, then set the name for debug purposes.
+			ThreadInitSyncEvent->Wait(INFINITE);
+			SetThreadName(ThreadID, TCHAR_TO_ANSI(*ThreadName));
+
+			ThreadPriority = TPri_Normal;
+			SetThreadPriority(InThreadPri);
+		}
+
+		// Cleanup the sync event
+		FPlatformProcess::ReturnSynchEventToPool(ThreadInitSyncEvent);
+		ThreadInitSyncEvent = nullptr;
+		return Thread != NULL;
+	}
+};
+#include "Microsoft/HideMicrosoftPlatformTypes.h"
diff --git a/Engine/Platforms/Hololens/Source/Runtime/Core/Public/HoloLensCriticalSection.h b/Engine/Platforms/Hololens/Source/Runtime/Core/Public/HoloLensCriticalSection.h
new file mode 100644
--- /dev/null
+++ b/Engine/Platforms/Hololens/Source/Runtime/Core/Public/HoloLensCriticalSection.h
@@ -0,0 +1,133 @@
+// Copyright Epic Games, Inc. All Rights Reserved.
+
+#pragma once
+
+#include "CoreTypes.h"
+#include "Misc/Timespan.h"
+#include "HAL/PlatformMemory.h"
+#include "GenericPlatform/GenericPlatformCriticalSection.h"
+
+/**
+ * This is the HoloLens version of a critical section. It uses an aggregate
+ * CRITICAL_SECTION to implement its locking.
+ */
+class FHoloLensCriticalSection
+{
+	/**
+	 * The HoloLens specific critical section
+	 */
+	Windows::CRITICAL_SECTION CriticalSection;
+
+public:
+
+	/**
+	 * Constructor that initializes the aggregated critical section
+	 */
+	FORCEINLINE FHoloLensCriticalSection()
+	{
+		CA_SUPPRESS(28125);
+		Windows::InitializeCriticalSection(&CriticalSection);
+		Windows::SetCriticalSectionSpinCount(&CriticalSection, 4000);
+	}
+
+	/**
+	 * Destructor cleaning up the critical section
+	 */
+	FORCEINLINE ~FHoloLensCriticalSection()
+	{
+		Windows::DeleteCriticalSection(&CriticalSection);
+	}
+
+	/**
+	 * Locks the critical section
+	 */
+	FORCEINLINE void Lock()
+	{
+		// Spin first before entering critical section, causing ring-0 transition and context switch.
+		if( Windows::TryEnterCriticalSection(&CriticalSection) == 0 )
+		{
+			Windows::EnterCriticalSection(&CriticalSection);
+		}
+	}
+
+	/**
+	 * Attempt to take a lock and returns whether or not a lock was taken.
+	 *
+	 * @return true if a lock was taken, false otherwise.
+	*/
+	FORCEINLINE bool TryLock()
+	{
+		if (Windows::TryEnterCriticalSection(&CriticalSection))
+		{
+			return true;
+		};
+		return false;
+	}
+
+	/**
+	 * Releases the lock on the critical seciton
+	 */
+	FORCEINLINE void Unlock()
+	{
+		Windows::LeaveCriticalSection(&CriticalSection);
+	}
+};
+
+/**
+ * FHoloLensRWLock - Read/Write Mutex
+ *	- Provides non-recursive Read/Write (or shared-exclusive) access.
+ *	- Windows specific lock structures/calls Ref: https://msdn.microsoft.com/en-us/library/windows/desktop/aa904937(v=vs.85).aspx
+ */
+class FHoloLensRWLock
+{
+public:
+	FHoloLensRWLock(const FHoloLensRWLock&) = delete;
+	FHoloLensRWLock& operator=(const FHoloLensRWLock&) = delete;
+
+	FORCEINLINE FHoloLensRWLock(uint32 Level = 0)
+	{
+		Windows::InitializeSRWLock(&Mutex);
+	}
+
+	FORCEINLINE ~FHoloLensRWLock()
+	{
+	}
+
+	FORCEINLINE void ReadLock()
+	{
+		Windows::AcquireSRWLockShared(&Mutex);
+	}
+
+	FORCEINLINE void WriteLock()
+	{
+		Windows::AcquireSRWLockExclusive(&Mutex);
+	}
+
+	FORCEINLINE bool TryReadLock()
+	{
+		return !!Windows::TryAcquireSRWLockShared(&Mutex);
+	}
+
+	FORCEINLINE bool TryWriteLock()
+	{
+		return !!Windows::TryAcquireSRWLockExclusive(&Mutex);
+	}
+
+	FORCEINLINE void ReadUnlock()
+	{
+		Windows::ReleaseSRWLockShared(&Mutex);
+	}
+
+	FORCEINLINE void WriteUnlock()
+	{
+		Windows::ReleaseSRWLockExclusive(&Mutex);
+	}
+
+private:
+	Windows::SRWLOCK Mutex;
+};
+
+
+typedef FHoloLensCriticalSection FCriticalSection;
+typedef FSystemWideCriticalSectionNotImplemented FSystemWideCriticalSection;
+typedef FHoloLensRWLock FRWLock;
diff --git a/Engine/Platforms/Hololens/Source/Runtime/Core/Public/HoloLensIncludes.h b/Engine/Platforms/Hololens/Source/Runtime/Core/Public/HoloLensIncludes.h
new file mode 100644
--- /dev/null
+++ b/Engine/Platforms/Hololens/Source/Runtime/Core/Public/HoloLensIncludes.h
@@ -0,0 +1,30 @@
+// Copyright Epic Games, Inc. All Rights Reserved.
+
+/*================================================================================
+	HoloLensIncludes.h: Includes the platform specific headers for HoloLens
+==================================================================================*/
+
+#pragma once
+
+// Set up compiler pragmas, etc
+#include "HoloLensCompilerSetup.h"
+
+// include platform implementations
+#include "HoloLensMemory.h"
+#include "HoloLensString.h"
+#include "HoloLensMisc.h"
+#include "HoloLensStackWalk.h"
+#include "HoloLensMath.h"
+#include "HoloLensTime.h"
+#include "HoloLensProcess.h"
+#include "HoloLensOutputDevices.h"
+#include "HoloLensAtomics.h"
+#include "HoloLensTLS.h"
+#include "HoloLensSplash.h"
+#include "HoloLensSurvey.h"
+
+typedef FGenericPlatformRHIFramePacer FPlatformRHIFramePacer;
+
+typedef FGenericPlatformAffinity FPlatformAffinity;
+
+#include "HoloLensProperties.h"
diff --git a/Engine/Platforms/Hololens/Source/Runtime/Core/Public/HoloLensPlatform.h b/Engine/Platforms/Hololens/Source/Runtime/Core/Public/HoloLensPlatform.h
new file mode 100644
--- /dev/null
+++ b/Engine/Platforms/Hololens/Source/Runtime/Core/Public/HoloLensPlatform.h
@@ -0,0 +1,124 @@
+// Copyright Epic Games, Inc. All Rights Reserved.
+
+/*================================================================================
+	HoloLensPlatform.h: Setup for the windows platform
+==================================================================================*/
+
+#pragma once
+
+#include "MSVC/MSVCPlatform.h"
+
+/** Define the HoloLens platform to be the active one **/
+#define PLATFORM_HOLOLENS					1
+
+/**
+* Windows specific types
+**/
+struct FHoloLensTypes : public FGenericPlatformTypes
+{
+	// defined in windefs.h, even though this is equivalent, the compiler doesn't think so
+	typedef unsigned long       DWORD;
+#ifdef _WIN64
+	typedef unsigned __int64	SIZE_T;
+	typedef __int64				SSIZE_T;
+#else
+	typedef unsigned long		SIZE_T;
+	typedef long				SSIZE_T;
+#endif
+
+	typedef decltype(__nullptr) TYPE_OF_NULLPTR;
+};
+
+typedef FHoloLensTypes FPlatformTypes;
+
+// Base defines, must define these for the platform, there are no defaults
+#define PLATFORM_DESKTOP				0
+#if defined( _WIN64 )
+#define PLATFORM_64BITS					1
+#else
+#define PLATFORM_64BITS					0
+#endif
+
+#define PLATFORM_CAN_SUPPORT_EDITORONLY_DATA	0
+
+// Base defines, defaults are commented out
+
+#define PLATFORM_LITTLE_ENDIAN								1
+#define PLATFORM_SUPPORTS_PRAGMA_PACK						1
+//@MIXEDREALITY_CHANGE : BEGIN 0 -> 1 to not use SSE2
+#if (defined(_M_ARM) || defined(_M_ARM64))
+#define PLATFORM_CPU_ARM_FAMILY								1
+#define PLATFORM_ENABLE_VECTORINTRINSICS					1
+#define PLATFORM_ENABLE_VECTORINTRINSICS_NEON				1
+#elif (defined(_M_IX86) || defined(_M_X64))
+#define PLATFORM_CPU_X86_FAMILY								1
+#define PLATFORM_ENABLE_VECTORINTRINSICS					1
+#endif
+#define PLATFORM_TASKGRAPH_GO_WIDE                          1
+//#define PLATFORM_USE_LS_SPEC_FOR_WIDECHAR					0
+//#define PLATFORM_COMPILER_DISTINGUISHES_DWORD_AND_UINT		1
+//#define PLATFORM_COMPILER_HAS_GENERIC_KEYWORD				1
+#define PLATFORM_HAS_BSD_TIME								0
+#define PLATFORM_HAS_BSD_SOCKETS							1
+#define PLATFORM_HAS_BSD_IPV6_SOCKETS						1
+#define PLATFORM_USE_PTHREADS								0
+#define PLATFORM_MAX_FILEPATH_LENGTH						MAX_PATH
+#define PLATFORM_USES_DYNAMIC_RHI							1
+#define PLATFORM_REQUIRES_FILESERVER						1
+#define PLATFORM_SUPPORTS_MULTITHREADED_GC					0
+#define PLATFORM_HAS_BSD_SOCKET_FEATURE_WINSOCKETS			1
+#define PLATFORM_USES_MICROSOFT_LIBC_FUNCTIONS				1
+//#define PLATFORM_SUPPORTS_NAMED_PIPES						1
+#define PLATFORM_COMPILER_HAS_VARIADIC_TEMPLATES			1
+#define PLATFORM_COMPILER_HAS_EXPLICIT_OPERATORS			1
+//#define EXCEPTION_EXECUTE_HANDLER                           1
+//#define PLATFORM_COMPILER_HAS_DEFAULT_FUNCTION_TEMPLATE_ARGUMENTS	1
+#define PLATFORM_USE_GENERIC_STRING_IMPLEMENTATION			0  // GenericWidePlatformString will fail on HoloLens in UE5
+
+#define PLATFORM_BREAK() (__nop(), __debugbreak())
+
[REDACTED_EMAIL]: Fixup once sockets are supported
+#define PLATFORM_SUPPORTS_MESSAGEBUS						1
+
+// @SPLASH_DAMAGE_CHANGE: [REDACTED_EMAIL] - BEGIN
+#define PLATFORM_SUPPORTS_XBOX_LIVE 0
+// @SPLASH_DAMAGE_CHANGE: [REDACTED_EMAIL] - END
+
+#define PLATFORM_HAS_128BIT_ATOMICS							PLATFORM_64BITS
+
+// Function type macros.
+#define VARARGS			__cdecl					/* Functions with variable arguments */
+#define CDECL			__cdecl					/* Standard C function */
+#define STDCALL			__stdcall				/* Standard calling convention */
+#define FORCEINLINE		__forceinline			/* Force code to be inline */
+#define FORCENOINLINE	__declspec(noinline)	/* Force code to NOT be inline */
+
+#define DECLARE_UINT64(x)	x
+
+// Optimization macros (uses __pragma to enable inside a #define).
+#define PRAGMA_DISABLE_OPTIMIZATION_ACTUAL __pragma(optimize("",off))
+#define PRAGMA_ENABLE_OPTIMIZATION_ACTUAL  __pragma(optimize("",on))
+
+// Backwater of the spec. All compilers support this except microsoft, and they will soon
+#define TYPENAME_OUTSIDE_TEMPLATE
+
+#pragma warning(disable : 4481) // nonstandard extension used: override specifier 'override'
+#define ABSTRACT abstract
+#define CONSTEXPR constexpr
+
+// Strings.
+#define LINE_TERMINATOR		TEXT("\r\n")
+#define LINE_TERMINATOR_ANSI "\r\n"
+
+// Alignment.
+#define MS_ALIGN(n) __declspec(align(n))
+
+// Pragmas
+#define MSVC_PRAGMA(Pragma) __pragma(Pragma)
+
+// DLL export and import definitions
+#define DLLEXPORT __declspec(dllexport)
+#define DLLIMPORT __declspec(dllimport)
+
+// disable this now as it is annoying for generic platform implementations
+#pragma warning(disable : 4100) // unreferenced formal parameter
diff --git a/Engine/Platforms/Hololens/Source/Runtime/Core/Public/HoloLensPlatformAffinity.h b/Engine/Platforms/Hololens/Source/Runtime/Core/Public/HoloLensPlatformAffinity.h
new file mode 100644
--- /dev/null
+++ b/Engine/Platforms/Hololens/Source/Runtime/Core/Public/HoloLensPlatformAffinity.h
@@ -0,0 +1,6 @@
+// Copyright Epic Games, Inc. All Rights Reserved.
+
+#pragma once
+
+#include "GenericPlatform/GenericPlatformAffinity.h"
+typedef FGenericPlatformAffinity FPlatformAffinity;
diff --git a/Engine/Platforms/Hololens/Source/Runtime/Core/Public/HoloLensPlatformCrashContext.h b/Engine/Platforms/Hololens/Source/Runtime/Core/Public/HoloLensPlatformCrashContext.h
new file mode 100644
--- /dev/null
+++ b/Engine/Platforms/Hololens/Source/Runtime/Core/Public/HoloLensPlatformCrashContext.h
@@ -0,0 +1,26 @@
+// Copyright Epic Games, Inc. All Rights Reserved.
+
+#include "GenericPlatform/GenericPlatformCrashContext.h"
+
+#include "Microsoft/AllowMicrosoftPlatformTypes.h"
+#include "DbgHelp.h"
+#include "Microsoft/HideMicrosoftPlatformTypes.h"
+
+#if WER_CUSTOM_REPORTS
+struct CORE_API FHoloLensPlatformCrashContext : public FGenericCrashContext
+{
+	/** Platform specific constants. */
+	enum EConstants
+	{
+		UE4_MINIDUMP_CRASHCONTEXT = LastReservedStream + 1,
+	};
+
+	virtual void AddPlatformSpecificProperties() override
+	{
+		AddCrashProperty( TEXT( "Platform.IsRunningHoloLens" ), 1 );
+	}
+};
+
+typedef FHoloLensPlatformCrashContext FPlatformCrashContext;
+
+#endif
\ No newline at end of file
diff --git a/Engine/Platforms/Hololens/Source/Runtime/Core/Public/HoloLensPlatformIncludes.h b/Engine/Platforms/Hololens/Source/Runtime/Core/Public/HoloLensPlatformIncludes.h
new file mode 100644
--- /dev/null
+++ b/Engine/Platforms/Hololens/Source/Runtime/Core/Public/HoloLensPlatformIncludes.h
@@ -0,0 +1,11 @@
+// Copyright Epic Games, Inc. All Rights Reserved.
+
+/*================================================================================
+	HoloLensPlatformIncludes.h: Includes the platform specific headers for HoloLens
+==================================================================================*/
+
+#pragma once
+#include "HoloLensSystemIncludes.h"
+
+// include platform implementations
+#include "HoloLensIncludes.h"
diff --git a/Engine/Platforms/Hololens/Source/Runtime/Core/Public/HoloLensPlatformMath.h b/Engine/Platforms/Hololens/Source/Runtime/Core/Public/HoloLensPlatformMath.h
new file mode 100644
--- /dev/null
+++ b/Engine/Platforms/Hololens/Source/Runtime/Core/Public/HoloLensPlatformMath.h
@@ -0,0 +1,11 @@
+// Copyright Epic Games, Inc. All Rights Reserved.
+
+/*=============================================================================================
+	HoloLensMath.h: HoloLens platform Math functions
+==============================================================================================*/
+
+#pragma once
+#include "Microsoft/MicrosoftPlatformMath.h"
+
+typedef FMicrosoftPlatformMathBase FHololensPlatformMath;
+typedef FHololensPlatformMath FPlatformMath;
diff --git a/Engine/Platforms/Hololens/Source/Runtime/Core/Public/HoloLensPlatformProcess.h b/Engine/Platforms/Hololens/Source/Runtime/Core/Public/HoloLensPlatformProcess.h
new file mode 100644
--- /dev/null
+++ b/Engine/Platforms/Hololens/Source/Runtime/Core/Public/HoloLensPlatformProcess.h
@@ -0,0 +1,72 @@
+// Copyright Epic Games, Inc. All Rights Reserved.
+
+/*=============================================================================================
+	HoloLensProcess.h: HoloLens platform Process functions
+==============================================================================================*/
+
+#pragma once
+#include "GenericPlatform/GenericPlatformProcess.h"
+
+/** Dummy process handle for platforms that use generic implementation. */
+struct FProcHandle : public TProcHandle<void*, nullptr>
+{
+public:
+	/** Default constructor. */
+	FORCEINLINE FProcHandle()
+		: TProcHandle()
+	{}
+
+	/** Initialization constructor. */
+	FORCEINLINE explicit FProcHandle( HandleType Other )
+		: TProcHandle( Other )
+	{}
+};
+
+/**
+ * HoloLens implementation of the Process OS functions
+ */
+struct CORE_API FHoloLensProcess : public FGenericPlatformProcess
+{
+	static const TCHAR* BaseDir();
+	static void Sleep( float Seconds );
+	static void SleepNoStats(float Seconds); 
+	[[noreturn]] static void SleepInfinite();
+	static void YieldThread();
+	UE_DEPRECATED(5.0, "Please use GetSynchEventFromPool to create a new event, and ReturnSynchEventToPool to release the event.")
+	static class FEvent* CreateSynchEvent(bool bIsManualReset = false);
+	static class FRunnableThread* CreateRunnableThread();
+	static const TCHAR* ExecutableName(bool bRemoveExtension = true);
+	static void* GetDllHandle(const TCHAR* Filename);
+	static void FreeDllHandle(void* DllHandle);
+	static void* GetDllExport(void* DllHandle, const TCHAR* ProcName);
+	static void PushDllDirectory(const TCHAR* Directory);
+	static void PopDllDirectory(const TCHAR* Directory);
+	static void SetCurrentWorkingDirectoryToBaseDir();
+	static FString GetCurrentWorkingDirectory();
+	/** Content saved to the game or engine directories should be rerouted to user directories instead **/
+	static bool ShouldSaveToUserDir();
+
+	static void SetThreadAffinityMask( uint64 AffinityMask );
+	static const TCHAR* UserDir();
+	static const TCHAR* UserSettingsDir();
+	static const TCHAR* UserTempDir();
+	static const TCHAR* ApplicationSettingsDir();
+
+	static const TCHAR* GetLocalAppDataLowLevelPath();
+	static const TCHAR* GetTempAppDataLowLevelPath();
+	static const TCHAR* GetLocalAppDataRedirectPath();
+	static const TCHAR* GetTempAppDataRedirectPath();
+	
+	static bool CanLaunchURL(const TCHAR* URL);
+	static void LaunchURL(const TCHAR* URL, const TCHAR* Parms, FString* Error);
+
+private:
+	/**
+	* Since Windows can only have one directory at a time,
+	* this stack is used to reset the previous DllDirectory.
+	*/
+	static TArray<FString> DllDirectoryStack;
+};
+
+typedef FHoloLensProcess FPlatformProcess;
+
diff --git a/Engine/Platforms/Hololens/Source/Runtime/Core/Public/HoloLensPlatformStackWalk.h b/Engine/Platforms/Hololens/Source/Runtime/Core/Public/HoloLensPlatformStackWalk.h
new file mode 100644
--- /dev/null
+++ b/Engine/Platforms/Hololens/Source/Runtime/Core/Public/HoloLensPlatformStackWalk.h
@@ -0,0 +1,16 @@
+// Copyright Epic Games, Inc. All Rights Reserved.
+
+/*=============================================================================================
+	HoloLensStackWalk.h: HoloLens platform stack walk functions
+==============================================================================================*/
+
+#pragma once
+#include "GenericPlatform/GenericPlatformStackWalk.h"
+
+struct CORE_API FHololensPlatformStackWalk
+	: public FGenericPlatformStackWalk
+{
+	static TArray<FProgramCounterSymbolInfo> GetStack(int32 IgnoreCount, int32 MaxDepth = 100, void* Context = nullptr);
+};
+
+typedef FHololensPlatformStackWalk FPlatformStackWalk;
diff --git a/Engine/Platforms/Hololens/Source/Runtime/Core/Public/HoloLensSystemIncludes.h b/Engine/Platforms/Hololens/Source/Runtime/Core/Public/HoloLensSystemIncludes.h
new file mode 100644
--- /dev/null
+++ b/Engine/Platforms/Hololens/Source/Runtime/Core/Public/HoloLensSystemIncludes.h
@@ -0,0 +1,36 @@
+// Copyright Epic Games, Inc. All Rights Reserved.
+#pragma once
+
+
+#include "CoreTypes.h"
+
+#ifndef STRICT
+#define STRICT
+#endif
+
+#include "HoloLensPlatformCompilerSetup.h"
+#include "Windows/MinimalWindowsApi.h"
+
+// Macro for releasing COM objects
+#define SAFE_RELEASE(p) { if(p) { (p)->Release(); (p)=NULL; } }
+THIRD_PARTY_INCLUDES_START
+
+// SIMD intrinsics
+#include <intrin.h>
+
+#include <intsafe.h>
+#include <strsafe.h>
+#include <tchar.h>
+
+#include <stdint.h>
+#include <concrt.h>
+
+THIRD_PARTY_INCLUDES_END
+
+#ifndef OUT
+#define OUT
+#endif
+
+#ifndef IN
+#define IN
+#endif
diff --git a/Engine/Platforms/Hololens/Source/Runtime/Core/Public/Microsoft/AllowMicrosoftPlatformAtomicsPrivate.h b/Engine/Platforms/Hololens/Source/Runtime/Core/Public/Microsoft/AllowMicrosoftPlatformAtomicsPrivate.h
new file mode 100644
--- /dev/null
+++ b/Engine/Platforms/Hololens/Source/Runtime/Core/Public/Microsoft/AllowMicrosoftPlatformAtomicsPrivate.h
@@ -0,0 +1,28 @@
+// Copyright Epic Games, Inc. All Rights Reserved.
+
+#ifndef WINDOWS_PLATFORM_ATOMICS_GUARD
+	#define WINDOWS_PLATFORM_ATOMICS_GUARD
+#else
+	#error Nesting AllowWindowsPlatformAtomics.h is not allowed!
+#endif
+
+#define InterlockedIncrement _InterlockedIncrement
+#define InterlockedDecrement _InterlockedDecrement
+#define InterlockedAdd _InterlockedAdd
+#define InterlockedExchange _InterlockedExchange
+#define InterlockedExchangeAdd _InterlockedExchangeAdd
+#define InterlockedCompareExchange _InterlockedCompareExchange
+
+#if PLATFORM_64BITS
+	#define InterlockedCompareExchangePointer _InterlockedCompareExchangePointer
+#else
+	#define InterlockedCompareExchangePointer __InlineInterlockedCompareExchangePointer
+#endif
+
+#if PLATFORM_64BITS
+	#define InterlockedExchange64 _InterlockedExchange64
+	#define InterlockedExchangeAdd64 _InterlockedExchangeAdd64
+	#define InterlockedCompareExchange64 _InterlockedCompareExchange64
+	#define InterlockedIncrement64 _InterlockedIncrement64
+	#define InterlockedDecrement64 _InterlockedDecrement64
+#endif
diff --git a/Engine/Platforms/Hololens/Source/Runtime/Core/Public/Microsoft/AllowMicrosoftPlatformTypesPrivate.h b/Engine/Platforms/Hololens/Source/Runtime/Core/Public/Microsoft/AllowMicrosoftPlatformTypesPrivate.h
new file mode 100644
--- /dev/null
+++ b/Engine/Platforms/Hololens/Source/Runtime/Core/Public/Microsoft/AllowMicrosoftPlatformTypesPrivate.h
@@ -0,0 +1,26 @@
+// Copyright Epic Games, Inc. All Rights Reserved.
+
+/*=============================================================================
+	AllowWindowsPLatformTypes.h: Defines for allowing the use of Windows types.
+=============================================================================*/
+
+#include "Microsoft/WindowsHWrapper.h"
+
+
+#ifndef WINDOWS_PLATFORM_TYPES_GUARD
+	#define WINDOWS_PLATFORM_TYPES_GUARD
+#else
+	#error Nesting AllowWindowsPlatformTypes.h is not allowed!
+#endif
+
+#define INT ::INT
+#define UINT ::UINT
+#define DWORD ::DWORD
+#define FLOAT ::FLOAT
+
+#ifndef TRUE
+#define TRUE 1
+#endif
+#ifndef FALSE
+#define FALSE 0
+#endif
diff --git a/Engine/Platforms/Hololens/Source/Runtime/Core/Public/Microsoft/HideMicrosoftPlatformAtomicsPrivate.h b/Engine/Platforms/Hololens/Source/Runtime/Core/Public/Microsoft/HideMicrosoftPlatformAtomicsPrivate.h
new file mode 100644
--- /dev/null
+++ b/Engine/Platforms/Hololens/Source/Runtime/Core/Public/Microsoft/HideMicrosoftPlatformAtomicsPrivate.h
@@ -0,0 +1,21 @@
+// Copyright Epic Games, Inc. All Rights Reserved.
+
+#ifdef WINDOWS_PLATFORM_ATOMICS_GUARD
+	#undef WINDOWS_PLATFORM_ATOMICS_GUARD
+#else
+	#error Mismatched HideWindowsPlatformAtomics.h detected.
+#endif
+
+#undef InterlockedIncrement
+#undef InterlockedDecrement
+#undef InterlockedAdd
+#undef InterlockedExchange
+#undef InterlockedExchangeAdd
+#undef InterlockedCompareExchange
+#undef InterlockedCompareExchangePointer
+#undef InterlockedExchange64
+#undef InterlockedExchangeAdd64
+#undef InterlockedCompareExchange64
+#undef InterlockedIncrement64
+#undef InterlockedDecrement64
+
diff --git a/Engine/Platforms/Hololens/Source/Runtime/Core/Public/Microsoft/HideMicrosoftPlatformTypesPrivate.h b/Engine/Platforms/Hololens/Source/Runtime/Core/Public/Microsoft/HideMicrosoftPlatformTypesPrivate.h
new file mode 100644
--- /dev/null
+++ b/Engine/Platforms/Hololens/Source/Runtime/Core/Public/Microsoft/HideMicrosoftPlatformTypesPrivate.h
@@ -0,0 +1,16 @@
+// Copyright Epic Games, Inc. All Rights Reserved.
+
+/*=============================================================================
+	HideWindowsPlatformTypes.h: Defines for hiding Windows type names.
+=============================================================================*/
+
+#ifdef WINDOWS_PLATFORM_TYPES_GUARD
+	#undef WINDOWS_PLATFORM_TYPES_GUARD
+#else
+	#error Mismatched HideWindowsPlatformTypes.h detected.
+#endif
+
+#undef INT
+#undef UINT
+#undef DWORD
+#undef FLOAT
diff --git a/Engine/Platforms/Hololens/Source/Runtime/Core/Public/Microsoft/WindowsHWrapperPrivate.h b/Engine/Platforms/Hololens/Source/Runtime/Core/Public/Microsoft/WindowsHWrapperPrivate.h
new file mode 100644
--- /dev/null
+++ b/Engine/Platforms/Hololens/Source/Runtime/Core/Public/Microsoft/WindowsHWrapperPrivate.h
@@ -0,0 +1,12 @@
+// Copyright Epic Games, Inc. All Rights Reserved.
+
+#pragma once
+
+#include "CoreTypes.h"
+#include "HAL/PlatformMemory.h"
+#include "Microsoft/PreWindowsApiPrivate.h"
+#ifndef STRICT
+#define STRICT
+#endif
+#include "Microsoft/MinWindowsPrivate.h"
+#include "Microsoft/PostWindowsApiPrivate.h"
diff --git a/Engine/Platforms/Hololens/Source/Runtime/D3D12RHI/D3D12RHI_HoloLens.build.cs b/Engine/Platforms/Hololens/Source/Runtime/D3D12RHI/D3D12RHI_HoloLens.build.cs
new file mode 100644
--- /dev/null
+++ b/Engine/Platforms/Hololens/Source/Runtime/D3D12RHI/D3D12RHI_HoloLens.build.cs
@@ -0,0 +1,14 @@
+// Copyright Epic Games, Inc. All Rights Reserved.
+
+namespace UnrealBuildTool.Rules
+{
+	public class D3D12RHI_HoloLens : D3D12RHI
+	{
+		protected override bool bUsesWindowsD3D12 { get => true; }
+
+		public D3D12RHI_HoloLens(ReadOnlyTargetRules Target) : base(Target)
+		{
+			PrivateIncludePaths.Add("../Platforms/HoloLens/Source/Runtime/D3D12RHI/Private");
+		}
+	}
+}
diff --git a/Engine/Platforms/Hololens/Source/Runtime/D3D12RHI/Private/D3D12RHIBasePrivate.h b/Engine/Platforms/Hololens/Source/Runtime/D3D12RHI/Private/D3D12RHIBasePrivate.h
new file mode 100644
--- /dev/null
+++ b/Engine/Platforms/Hololens/Source/Runtime/D3D12RHI/Private/D3D12RHIBasePrivate.h
@@ -0,0 +1,25 @@
+// Copyright Epic Games, Inc. All Rights Reserved.
+
+/*=============================================================================
+	D3D12BaseRHIPrivate.h: Private D3D RHI definitions for HoloLens.
+=============================================================================*/
+
+#pragma once
+
+#include "D3D12ThirdParty.h"
+
+#define D3D12RHI_RESOURCE_FLAG_ALLOW_INDIRECT_BUFFER D3D12_RESOURCE_FLAG_NONE
+#define D3D12RHI_HEAP_FLAG_ALLOW_INDIRECT_BUFFERS		D3D12_HEAP_FLAG_NONE
+
+#define D3D12RHI_NEEDS_VENDOR_EXTENSIONS     0
+#define D3D12RHI_NEEDS_SHADER_FEATURE_CHECKS 0
+
+// Heap create not zeroed flag is not available on Hololens so use internal define to disable the feature
+// but make code path shared when it becomes available
+#define FD3D12_HEAP_FLAG_CREATE_NOT_ZEROED D3D12_HEAP_FLAG_NONE
+
+#include "../Public/D3D12Util.h"
+
+// Windows desktop version is used.
+#include "Windows/WindowsD3D12DiskCache.h"
+#include "Windows/WindowsD3D12PipelineState.h"
diff --git a/Engine/Platforms/Hololens/Source/Runtime/D3D12RHI/Private/HoloLensD3D12Device.cpp b/Engine/Platforms/Hololens/Source/Runtime/D3D12RHI/Private/HoloLensD3D12Device.cpp
new file mode 100644
--- /dev/null
+++ b/Engine/Platforms/Hololens/Source/Runtime/D3D12RHI/Private/HoloLensD3D12Device.cpp
@@ -0,0 +1,6 @@
+// Copyright Epic Games, Inc. All Rights Reserved.
+
+#include "D3D12RHIPrivate.h"
+
+// Same as Windows desktop version
+#include "Windows/WindowsD3D12Device.cpp"
\ No newline at end of file
diff --git a/Engine/Platforms/Hololens/Source/Runtime/D3D12RHI/Private/HoloLensD3D12DiskCache.cpp b/Engine/Platforms/Hololens/Source/Runtime/D3D12RHI/Private/HoloLensD3D12DiskCache.cpp
new file mode 100644
--- /dev/null
+++ b/Engine/Platforms/Hololens/Source/Runtime/D3D12RHI/Private/HoloLensD3D12DiskCache.cpp
@@ -0,0 +1,6 @@
+// Copyright Epic Games, Inc. All Rights Reserved.
+
+#include "D3D12RHIPrivate.h"
+
+// Same as Windows desktop version
+#include "Windows/WindowsD3D12DiskCache.cpp"
\ No newline at end of file
diff --git a/Engine/Platforms/Hololens/Source/Runtime/D3D12RHI/Private/HoloLensD3D12Fence.cpp b/Engine/Platforms/Hololens/Source/Runtime/D3D12RHI/Private/HoloLensD3D12Fence.cpp
new file mode 100644
--- /dev/null
+++ b/Engine/Platforms/Hololens/Source/Runtime/D3D12RHI/Private/HoloLensD3D12Fence.cpp
@@ -0,0 +1,11 @@
+// Copyright Epic Games, Inc. All Rights Reserved.
+
+// Implementation of D3D12 fence functions
+
+//-----------------------------------------------------------------------------
+//	Include Files
+//-----------------------------------------------------------------------------
+#include "D3D12RHIPrivate.h"
+
+// Same as Windows desktop version
+#include "Windows/WindowsD3D12Fence.cpp"
\ No newline at end of file
diff --git a/Engine/Platforms/Hololens/Source/Runtime/D3D12RHI/Private/HoloLensD3D12PipelineState.cpp b/Engine/Platforms/Hololens/Source/Runtime/D3D12RHI/Private/HoloLensD3D12PipelineState.cpp
new file mode 100644
--- /dev/null
+++ b/Engine/Platforms/Hololens/Source/Runtime/D3D12RHI/Private/HoloLensD3D12PipelineState.cpp
@@ -0,0 +1,12 @@
+// Copyright Epic Games, Inc. All Rights Reserved.
+
+// Implementation of D3D12 Pipelinestate related functions
+
+//-----------------------------------------------------------------------------
+//	Include Files
+//-----------------------------------------------------------------------------
+#include "D3D12RHIPrivate.h"
+
+// Same as Windows desktop version
+#include "Windows/WindowsD3D12PipelineState.cpp"
+
diff --git a/Engine/Platforms/Hololens/Source/Runtime/D3D12RHI/Private/HoloLensD3D12Viewport.cpp b/Engine/Platforms/Hololens/Source/Runtime/D3D12RHI/Private/HoloLensD3D12Viewport.cpp
new file mode 100644
--- /dev/null
+++ b/Engine/Platforms/Hololens/Source/Runtime/D3D12RHI/Private/HoloLensD3D12Viewport.cpp
@@ -0,0 +1,11 @@
+// Copyright Epic Games, Inc. All Rights Reserved.
+
+/*=============================================================================
+	D3D12Viewport.cpp: D3D viewport RHI implementation.
+=============================================================================*/
+
+#include "D3D12RHIPrivate.h"
+#include "RenderCore.h"
+
+// Borrow the Windows version; HoloLens-specific edits are now inline - previous approach was starting to duplicate too much code.
+#include "Windows/WindowsD3D12Viewport.cpp"
\ No newline at end of file
diff --git a/Engine/Platforms/Hololens/Source/Runtime/InputCore/InputCore_HoloLens.build.cs b/Engine/Platforms/Hololens/Source/Runtime/InputCore/InputCore_HoloLens.build.cs
new file mode 100644
--- /dev/null
+++ b/Engine/Platforms/Hololens/Source/Runtime/InputCore/InputCore_HoloLens.build.cs
@@ -0,0 +1,11 @@
+// Copyright Epic Games, Inc. All Rights Reserved.
+
+namespace UnrealBuildTool.Rules
+{
+	public class InputCore_HoloLens : InputCore
+	{
+		public InputCore_HoloLens(ReadOnlyTargetRules Target) : base(Target)
+		{
+		}
+	}
+}
diff --git a/Engine/Platforms/Hololens/Source/Runtime/InputCore/Private/HoloLensPlatformInput.cpp b/Engine/Platforms/Hololens/Source/Runtime/InputCore/Private/HoloLensPlatformInput.cpp
new file mode 100644
--- /dev/null
+++ b/Engine/Platforms/Hololens/Source/Runtime/InputCore/Private/HoloLensPlatformInput.cpp
@@ -0,0 +1,94 @@
+// Copyright Epic Games, Inc. All Rights Reserved.
+
+#include "HoloLensPlatformInput.h"
+#include "Microsoft/WindowsHWrapper.h"
+
+uint32 FHoloLensPlatformInput::GetCharKeyMap(uint32* KeyCodes, FString* KeyNames, uint32 MaxMappings)
+{
+	// Map lowercase keys since HoloLensInputInterface::HoloLensMapVirtualKeyToCharacter is sensitive to Shift state
+	// Win32 MapVirtualKey is not, but its behavior is annoying to reproduce since VK_1 maps to '1' (no shift)
+	// whereas VK_W maps to 'W' (caps).
+	return FGenericPlatformInput::GetStandardPrintableKeyMap(KeyCodes, KeyNames, MaxMappings, true, true);
+}
+
+uint32 FHoloLensPlatformInput::GetKeyMap(uint32* KeyCodes, FString* KeyNames, uint32 MaxMappings)
+{
+#define ADDKEYMAP(KeyCode, KeyName)		if (NumMappings<MaxMappings) { KeyCodes[NumMappings]=KeyCode; KeyNames[NumMappings]=KeyName; ++NumMappings; };
+
+	uint32 NumMappings = 0;
+
+	if (KeyCodes && KeyNames && (MaxMappings > 0))
+	{
+		ADDKEYMAP(VK_LBUTTON, TEXT("LeftMouseButton"));
+		ADDKEYMAP(VK_RBUTTON, TEXT("RightMouseButton"));
+		ADDKEYMAP(VK_MBUTTON, TEXT("MiddleMouseButton"));
+
+		ADDKEYMAP(VK_XBUTTON1, TEXT("ThumbMouseButton"));
+		ADDKEYMAP(VK_XBUTTON2, TEXT("ThumbMouseButton2"));
+
+		ADDKEYMAP(VK_BACK, TEXT("BackSpace"));
+		ADDKEYMAP(VK_TAB, TEXT("Tab"));
+		ADDKEYMAP(VK_RETURN, TEXT("Enter"));
+		ADDKEYMAP(VK_PAUSE, TEXT("Pause"));
+
+		ADDKEYMAP(VK_CAPITAL, TEXT("CapsLock"));
+		ADDKEYMAP(VK_ESCAPE, TEXT("Escape"));
+		ADDKEYMAP(VK_SPACE, TEXT("SpaceBar"));
+		ADDKEYMAP(VK_PRIOR, TEXT("PageUp"));
+		ADDKEYMAP(VK_NEXT, TEXT("PageDown"));
+		ADDKEYMAP(VK_END, TEXT("End"));
+		ADDKEYMAP(VK_HOME, TEXT("Home"));
+
+		ADDKEYMAP(VK_LEFT, TEXT("Left"));
+		ADDKEYMAP(VK_UP, TEXT("Up"));
+		ADDKEYMAP(VK_RIGHT, TEXT("Right"));
+		ADDKEYMAP(VK_DOWN, TEXT("Down"));
+
+		ADDKEYMAP(VK_INSERT, TEXT("Insert"));
+		ADDKEYMAP(VK_DELETE, TEXT("Delete"));
+
+		ADDKEYMAP(VK_NUMPAD0, TEXT("NumPadZero"));
+		ADDKEYMAP(VK_NUMPAD1, TEXT("NumPadOne"));
+		ADDKEYMAP(VK_NUMPAD2, TEXT("NumPadTwo"));
+		ADDKEYMAP(VK_NUMPAD3, TEXT("NumPadThree"));
+		ADDKEYMAP(VK_NUMPAD4, TEXT("NumPadFour"));
+		ADDKEYMAP(VK_NUMPAD5, TEXT("NumPadFive"));
+		ADDKEYMAP(VK_NUMPAD6, TEXT("NumPadSix"));
+		ADDKEYMAP(VK_NUMPAD7, TEXT("NumPadSeven"));
+		ADDKEYMAP(VK_NUMPAD8, TEXT("NumPadEight"));
+		ADDKEYMAP(VK_NUMPAD9, TEXT("NumPadNine"));
+
+		ADDKEYMAP(VK_MULTIPLY, TEXT("Multiply"));
+		ADDKEYMAP(VK_ADD, TEXT("Add"));
+		ADDKEYMAP(VK_SUBTRACT, TEXT("Subtract"));
+		ADDKEYMAP(VK_DECIMAL, TEXT("Decimal"));
+		ADDKEYMAP(VK_DIVIDE, TEXT("Divide"));
+
+		ADDKEYMAP(VK_F1, TEXT("F1"));
+		ADDKEYMAP(VK_F2, TEXT("F2"));
+		ADDKEYMAP(VK_F3, TEXT("F3"));
+		ADDKEYMAP(VK_F4, TEXT("F4"));
+		ADDKEYMAP(VK_F5, TEXT("F5"));
+		ADDKEYMAP(VK_F6, TEXT("F6"));
+		ADDKEYMAP(VK_F7, TEXT("F7"));
+		ADDKEYMAP(VK_F8, TEXT("F8"));
+		ADDKEYMAP(VK_F9, TEXT("F9"));
+		ADDKEYMAP(VK_F10, TEXT("F10"));
+		ADDKEYMAP(VK_F11, TEXT("F11"));
+		ADDKEYMAP(VK_F12, TEXT("F12"));
+
+		ADDKEYMAP(VK_NUMLOCK, TEXT("NumLock"));
+
+		ADDKEYMAP(VK_SCROLL, TEXT("ScrollLock"));
+
+		ADDKEYMAP(VK_LSHIFT, TEXT("LeftShift"));
+		ADDKEYMAP(VK_RSHIFT, TEXT("RightShift"));
+		ADDKEYMAP(VK_LCONTROL, TEXT("LeftControl"));
+		ADDKEYMAP(VK_RCONTROL, TEXT("RightControl"));
+		ADDKEYMAP(VK_LMENU, TEXT("LeftAlt"));
+		ADDKEYMAP(VK_RMENU, TEXT("RightAlt"));
+	}
+
+	check(NumMappings < MaxMappings);
+	return NumMappings;
+}
diff --git a/Engine/Platforms/Hololens/Source/Runtime/InputCore/Public/HoloLensPlatformInput.h b/Engine/Platforms/Hololens/Source/Runtime/InputCore/Public/HoloLensPlatformInput.h
new file mode 100644
--- /dev/null
+++ b/Engine/Platforms/Hololens/Source/Runtime/InputCore/Public/HoloLensPlatformInput.h
@@ -0,0 +1,14 @@
+// Copyright Epic Games, Inc. All Rights Reserved.
+
+#pragma once
+
+#include "CoreTypes.h"
+#include "GenericPlatform/GenericPlatformInput.h"
+
+struct INPUTCORE_API FHoloLensPlatformInput : FGenericPlatformInput
+{
+	static uint32 GetKeyMap(uint32* KeyCodes, FString* KeyNames, uint32 MaxMappings);
+	static uint32 GetCharKeyMap(uint32* KeyCodes, FString* KeyNames, uint32 MaxMappings);
+};
+
+typedef FHoloLensPlatformInput FPlatformInput;
\ No newline at end of file
diff --git a/Engine/Platforms/Hololens/Source/Runtime/Launch/Launch_HoloLens.build.cs b/Engine/Platforms/Hololens/Source/Runtime/Launch/Launch_HoloLens.build.cs
new file mode 100644
--- /dev/null
+++ b/Engine/Platforms/Hololens/Source/Runtime/Launch/Launch_HoloLens.build.cs
@@ -0,0 +1,13 @@
+// Copyright Epic Games, Inc. All Rights Reserved.
+
+namespace UnrealBuildTool.Rules
+{
+	public class Launch_HoloLens : Launch
+	{
+		public Launch_HoloLens(ReadOnlyTargetRules Target) : base(Target)
+		{
+			DynamicallyLoadedModuleNames.Add("D3D11RHI");
+			DynamicallyLoadedModuleNames.Add("AudioMixerXAudio2");
+		}
+	}
+}
