move HoloLens runtime source to platform extension folder.

diff --git a/Engine/Platforms/Hololens/Source/Runtime/Launch/Private/HoloLensLaunch.cpp b/Engine/Platforms/Hololens/Source/Runtime/Launch/Private/HoloLensLaunch.cpp
new file mode 100644
--- /dev/null
+++ b/Engine/Platforms/Hololens/Source/Runtime/Launch/Private/HoloLensLaunch.cpp
@@ -0,0 +1,1270 @@
+// Copyright Epic Games, Inc. All Rights Reserved.
+
+#include "CoreMinimal.h"
+
+#include "HAL/PlatformFileManager.h"
+#include "HAL/PlatformAffinity.h"
+#include "Misc/LocalTimestampDirectoryVisitor.h"
+
+#include "Misc/App.h"
+#include "Misc/CoreDelegates.h"
+#include "Misc/Paths.h"
+#include "Misc/PackageName.h"
+#include "Misc/EngineVersion.h"
+
+#include "LaunchEngineLoop.h"
+//#include "CoreTypes.h"
+#include "HAL/ExceptionHandling.h"
+#include "RHI.h"
+#include "RenderingThread.h"
+
+#include "Interfaces/IPluginManager.h"
+
+#include "HoloLensApplication.h"
+#include "HoloLensCursor.h"
+#include "HoloLensWindow.h"
+#include <stdio.h>
+
+//#include "Microsoft/AllowMicrosoftPlatformAtomics.h"
+#include "Microsoft/AllowMicrosoftPlatformTypes.h"
+#include <collection.h>
+#include <ppltasks.h>
+#include <concurrent_queue.h>
+#include "Microsoft/HideMicrosoftPlatformTypes.h"
+//#include "Microsoft/HideMicrosoftPlatformAtomics.h"
+
+// @MIXEDREALITY_CHANGE : BEGIN - Enable Stereo rendering for suspended apps
+#include "IHeadMountedDisplay.h"
+#include "IXRTrackingSystem.h"
+#include "Modules/ModuleManager.h"
+// @MIXEDREALITY_CHANGE : END
+
+// http://developer.download.nvidia.com/devzone/devcenter/gamegraphics/files/OptimusRenderingPolicies.pdf
+// The following line is to favor the high performance NVIDIA GPU if there are multiple GPUs
+// Has to be .exe module to be correctly detected.
+extern "C" { _declspec(dllexport) uint32 NvOptimusEnablement = 0x00000001; }
+
+// Similar for AMD GPU
+extern "C" {__declspec(dllexport) int AmdPowerXpressRequestHighPerformance = 1; }
+
+#pragma warning(disable : 4946)	// reinterpret_cast used between related classes: 'Platform::Object' and ...
+
+using namespace Windows::ApplicationModel::Activation;
+using namespace Windows::ApplicationModel::Core;
+using namespace Windows::ApplicationModel::ExtendedExecution;
+using namespace Windows::Foundation;
+using namespace Windows::System;
+using namespace Windows::UI::Core;
+using namespace Windows::Graphics::Holographic;
+
+static bool GSeparateCoreWindowVisible = false;
+
+DEFINE_LOG_CATEGORY_STATIC(LogLaunchHoloLens, Log, All);
+
+void appHoloLensEarlyInit();
+int32 GuardedMain( const TCHAR* CmdLine, HINSTANCE hInInstance, HINSTANCE hPrevInstance, int32 nCmdShow );
+
+ref class ViewProvider sealed : public Windows::ApplicationModel::Core::IFrameworkView
+{
+public:
+	ViewProvider();
+
+	// IFrameworkView impl.
+	virtual void Initialize(Windows::ApplicationModel::Core::CoreApplicationView^ applicationView);
+	virtual void SetWindow(Windows::UI::Core::CoreWindow^ window);
+	virtual void Load(Platform::String^ entryPoint);
+	virtual void Run();
+	virtual void Uninitialize();
+
+	// Process events on the core window dispatcher.
+	void ProcessEvents();
+
+	// Keyboard event handlers.
+	void OnKeyDown(_In_ CoreWindow^ sender, _In_ KeyEventArgs^ args);
+	void OnKeyUp(_In_ CoreWindow^ sender, _In_ KeyEventArgs^ args);
+	void OnCharacterReceived(_In_ CoreWindow^ sender, _In_ CharacterReceivedEventArgs^ args);
+
+	//! Dispatcher event handlers.
+	void OnAcceleratorKeyActivated(CoreDispatcher ^sender, AcceleratorKeyEventArgs ^args);
+	// Window event handlers.
+	void OnVisibilityChanged(_In_ CoreWindow^ sender, _In_ VisibilityChangedEventArgs^ args);
+	void OnWindowSizeChanged(Windows::UI::Core::CoreWindow^ sender, Windows::UI::Core::WindowSizeChangedEventArgs^ args);
+	void OnWindowClosed(Windows::UI::Core::CoreWindow^ sender, Windows::UI::Core::CoreWindowEventArgs^ args);
+
+	void OnPointerPressed(Windows::UI::Core::CoreWindow^ sender, Windows::UI::Core::PointerEventArgs^ args);
+	void OnPointerReleased(Windows::UI::Core::CoreWindow^ sender, Windows::UI::Core::PointerEventArgs^ args);
+	void OnPointerMoved(Windows::UI::Core::CoreWindow^ sender, Windows::UI::Core::PointerEventArgs^ args);
+	void OnPointerWheelChanged(Windows::UI::Core::CoreWindow^ sender, Windows::UI::Core::PointerEventArgs^ args);
+
+	bool IsVisible()
+	{
+		return bVisible;
+	}
+
+	bool WasWindowClosed()
+	{
+		return bWindowClosed;
+	}
+
+	void InitOptionalPackages();
+
+private:
+	Windows::Foundation::Rect existingSize;
+
+	bool ActivationComplete;
+
+	void OnActivated( _In_ Windows::ApplicationModel::Core::CoreApplicationView^ applicationView, _In_ Windows::ApplicationModel::Activation::IActivatedEventArgs^ args );
+	void OnResuming( _In_ Platform::Object^ sender, _In_ Platform::Object^ args );
+	void OnSuspending( _In_ Platform::Object^ sender, _In_ Windows::ApplicationModel::SuspendingEventArgs^ args );
+	void OnEnteredBackground(_In_ Platform::Object^ sender, _In_ Windows::ApplicationModel::EnteredBackgroundEventArgs^ args);
+	void OnLeavingBackground(_In_ Platform::Object^ sender, _In_ Windows::ApplicationModel::LeavingBackgroundEventArgs^ args);
+	void OnExiting(_In_ Platform::Object^ sender, _In_ Platform::Object^ args);
+
+	const TCHAR* GetPointerUpdateKindString(Windows::UI::Input::PointerUpdateKind InKind);
+	EMouseButtons::Type PointerUpdateKindToUEKey(Windows::UI::Input::PointerUpdateKind InKind, bool& bWasPressed);
+	bool ProcessMouseEvent(Windows::UI::Core::PointerEventArgs^ args);
+
+	bool bVisible;
+	bool bWindowClosed;
+
+	struct QueuedPointerEvent
+	{
+		QueuedPointerEvent(Windows::UI::Core::PointerEventArgs^ Args) :
+			RawPosition(Args->CurrentPoint->Position.X, Args->CurrentPoint->Position.Y),
+			WheelDelta(Args->CurrentPoint->Properties->MouseWheelDelta),
+			Kind(Args->CurrentPoint->Properties->PointerUpdateKind)
+		{
+		}
+
+		FVector2D RawPosition;
+		int32 WheelDelta;
+		Windows::UI::Input::PointerUpdateKind Kind;
+	};
+	bool ProcessMouseEvent(const QueuedPointerEvent& Event);
+
+	static void RegisterDlcPluginsForOptionalPackage(Windows::ApplicationModel::Package ^DependencyPackage);
+	static void EnsureDlcPluginsAreMounted();
+
+	TArray<QueuedPointerEvent> PointerEventQueue;
+};
+
+ViewProvider^ GViewProvider = nullptr;
+
+ref class ViewProviderFactory sealed : Windows::ApplicationModel::Core::IFrameworkViewSource 
+{
+public:
+	ViewProviderFactory() {}
+	virtual Windows::ApplicationModel::Core::IFrameworkView^ CreateView()
+	{
+		GViewProvider = ref new ViewProvider();
+		return GViewProvider;
+	}
+};
+
+ViewProvider::ViewProvider() :
+	ActivationComplete(false)
+{
+}
+
+uint32 TranslateWinRTKey(Windows::UI::Core::KeyEventArgs^ keyEventArgs)
+{
+	Windows::System::VirtualKey WinRTKeyCode = keyEventArgs->VirtualKey;
+
+	uint32 TranslatedKeyCode = (uint32)WinRTKeyCode;
+	switch( WinRTKeyCode )
+	{
+	case Windows::System::VirtualKey::LeftButton:	TranslatedKeyCode = VK_LBUTTON;		break;
+	case Windows::System::VirtualKey::RightButton:	TranslatedKeyCode = VK_RBUTTON;		break;
+	case Windows::System::VirtualKey::MiddleButton:	TranslatedKeyCode = VK_MBUTTON;		break;
+
+	case Windows::System::VirtualKey::XButton1:		TranslatedKeyCode = VK_XBUTTON1;	break;
+	case Windows::System::VirtualKey::XButton2:		TranslatedKeyCode = VK_XBUTTON2;	break;
+
+	case Windows::System::VirtualKey::Back:			TranslatedKeyCode = VK_BACK;		break;
+	case Windows::System::VirtualKey::Tab:			TranslatedKeyCode = VK_TAB;			break;
+	case Windows::System::VirtualKey::Enter:		TranslatedKeyCode = VK_RETURN;		break;
+	case Windows::System::VirtualKey::Pause:		TranslatedKeyCode = VK_PAUSE;		break;
+
+	case Windows::System::VirtualKey::CapitalLock:	TranslatedKeyCode = VK_CAPITAL;		break;
+	case Windows::System::VirtualKey::Escape:		TranslatedKeyCode = VK_ESCAPE;		break;
+	case Windows::System::VirtualKey::Space:		TranslatedKeyCode = VK_SPACE;		break;
+
+	case Windows::System::VirtualKey::PageUp:		TranslatedKeyCode = VK_PRIOR;		break;
+	case Windows::System::VirtualKey::PageDown:		TranslatedKeyCode = VK_NEXT;		break;
+	case Windows::System::VirtualKey::End:			TranslatedKeyCode = VK_END;			break;
+	case Windows::System::VirtualKey::Home:			TranslatedKeyCode = VK_HOME;		break;
+
+	case Windows::System::VirtualKey::Left:			TranslatedKeyCode = VK_LEFT;		break;
+	case Windows::System::VirtualKey::Up:			TranslatedKeyCode = VK_UP;			break;
+	case Windows::System::VirtualKey::Right:		TranslatedKeyCode = VK_RIGHT;		break;
+	case Windows::System::VirtualKey::Down:			TranslatedKeyCode = VK_DOWN;		break;
+
+	case Windows::System::VirtualKey::Insert:		TranslatedKeyCode = VK_INSERT;		break;
+	case Windows::System::VirtualKey::Delete:		TranslatedKeyCode = VK_DELETE;		break;
+
+	// VK_0 - VK_9 are the same as ASCII '0' - '9' (0x30 - 0x39)
+	// 0x40 : unassigned
+	// VK_A - VK_Z are the same as ASCII 'A' - 'Z' (0x41 - 0x5A)
+	case Windows::System::VirtualKey::Number0:		TranslatedKeyCode = 0x30;			break;
+	case Windows::System::VirtualKey::Number1:		TranslatedKeyCode = 0x31;			break;
+	case Windows::System::VirtualKey::Number2:		TranslatedKeyCode = 0x32;			break;
+	case Windows::System::VirtualKey::Number3:		TranslatedKeyCode = 0x33;			break;
+	case Windows::System::VirtualKey::Number4:		TranslatedKeyCode = 0x34;			break;
+	case Windows::System::VirtualKey::Number5:		TranslatedKeyCode = 0x35;			break;
+	case Windows::System::VirtualKey::Number6:		TranslatedKeyCode = 0x36;			break;
+	case Windows::System::VirtualKey::Number7:		TranslatedKeyCode = 0x37;			break;
+	case Windows::System::VirtualKey::Number8:		TranslatedKeyCode = 0x38;			break;
+	case Windows::System::VirtualKey::Number9:		TranslatedKeyCode = 0x39;			break;
+	case Windows::System::VirtualKey::A:			TranslatedKeyCode = 0x41;			break;
+	case Windows::System::VirtualKey::B:			TranslatedKeyCode = 0x42;			break;
+	case Windows::System::VirtualKey::C:			TranslatedKeyCode = 0x43;			break;
+	case Windows::System::VirtualKey::D:			TranslatedKeyCode = 0x44;			break;
+	case Windows::System::VirtualKey::E:			TranslatedKeyCode = 0x45;			break;
+	case Windows::System::VirtualKey::F:			TranslatedKeyCode = 0x46;			break;
+	case Windows::System::VirtualKey::G:			TranslatedKeyCode = 0x47;			break;
+	case Windows::System::VirtualKey::H:			TranslatedKeyCode = 0x48;			break;
+	case Windows::System::VirtualKey::I:			TranslatedKeyCode = 0x49;			break;
+	case Windows::System::VirtualKey::J:			TranslatedKeyCode = 0x4A;			break;
+	case Windows::System::VirtualKey::K:			TranslatedKeyCode = 0x4B;			break;
+	case Windows::System::VirtualKey::L:			TranslatedKeyCode = 0x4C;			break;
+	case Windows::System::VirtualKey::M:			TranslatedKeyCode = 0x4D;			break;
+	case Windows::System::VirtualKey::N:			TranslatedKeyCode = 0x4E;			break;
+	case Windows::System::VirtualKey::O:			TranslatedKeyCode = 0x4F;			break;
+	case Windows::System::VirtualKey::P:			TranslatedKeyCode = 0x50;			break;
+	case Windows::System::VirtualKey::Q:			TranslatedKeyCode = 0x51;			break;
+	case Windows::System::VirtualKey::R:			TranslatedKeyCode = 0x52;			break;
+	case Windows::System::VirtualKey::S:			TranslatedKeyCode = 0x53;			break;
+	case Windows::System::VirtualKey::T:			TranslatedKeyCode = 0x54;			break;
+	case Windows::System::VirtualKey::U:			TranslatedKeyCode = 0x55;			break;
+	case Windows::System::VirtualKey::V:			TranslatedKeyCode = 0x56;			break;
+	case Windows::System::VirtualKey::W:			TranslatedKeyCode = 0x57;			break;
+	case Windows::System::VirtualKey::X:			TranslatedKeyCode = 0x58;			break;
+	case Windows::System::VirtualKey::Y:			TranslatedKeyCode = 0x59;			break;
+	case Windows::System::VirtualKey::Z:			TranslatedKeyCode = 0x5A;			break;
+
+	case Windows::System::VirtualKey::NumberPad0:	TranslatedKeyCode = VK_NUMPAD0;		break;
+	case Windows::System::VirtualKey::NumberPad1:	TranslatedKeyCode = VK_NUMPAD1;		break;
+	case Windows::System::VirtualKey::NumberPad2:	TranslatedKeyCode = VK_NUMPAD2;		break;
+	case Windows::System::VirtualKey::NumberPad3:	TranslatedKeyCode = VK_NUMPAD3;		break;
+	case Windows::System::VirtualKey::NumberPad4:	TranslatedKeyCode = VK_NUMPAD4;		break;
+	case Windows::System::VirtualKey::NumberPad5:	TranslatedKeyCode = VK_NUMPAD5;		break;
+	case Windows::System::VirtualKey::NumberPad6:	TranslatedKeyCode = VK_NUMPAD6;		break;
+	case Windows::System::VirtualKey::NumberPad7:	TranslatedKeyCode = VK_NUMPAD7;		break;
+	case Windows::System::VirtualKey::NumberPad8:	TranslatedKeyCode = VK_NUMPAD8;		break;
+	case Windows::System::VirtualKey::NumberPad9:	TranslatedKeyCode = VK_NUMPAD9;		break;
+
+	case Windows::System::VirtualKey::Multiply:		TranslatedKeyCode = VK_MULTIPLY;	break;
+	case Windows::System::VirtualKey::Add:			TranslatedKeyCode = VK_ADD;			break;
+	case Windows::System::VirtualKey::Subtract:		TranslatedKeyCode = VK_SUBTRACT;	break;
+	case Windows::System::VirtualKey::Decimal:		TranslatedKeyCode = VK_DECIMAL;		break;
+	case Windows::System::VirtualKey::Divide:		TranslatedKeyCode = VK_DIVIDE;		break;
+
+	case Windows::System::VirtualKey::F1:			TranslatedKeyCode = VK_F1;			break;
+	case Windows::System::VirtualKey::F2:			TranslatedKeyCode = VK_F2;			break;
+	case Windows::System::VirtualKey::F3:			TranslatedKeyCode = VK_F3;			break;
+	case Windows::System::VirtualKey::F4:			TranslatedKeyCode = VK_F4;			break;
+	case Windows::System::VirtualKey::F5:			TranslatedKeyCode = VK_F5;			break;
+	case Windows::System::VirtualKey::F6:			TranslatedKeyCode = VK_F6;			break;
+	case Windows::System::VirtualKey::F7:			TranslatedKeyCode = VK_F7;			break;
+	case Windows::System::VirtualKey::F8:			TranslatedKeyCode = VK_F8;			break;
+	case Windows::System::VirtualKey::F9:			TranslatedKeyCode = VK_F9;			break;
+	case Windows::System::VirtualKey::F10:			TranslatedKeyCode = VK_F10;			break;
+	case Windows::System::VirtualKey::F11:			TranslatedKeyCode = VK_F11;			break;
+	case Windows::System::VirtualKey::F12:			TranslatedKeyCode = VK_F12;			break;
+
+	case Windows::System::VirtualKey::NumberKeyLock:TranslatedKeyCode = VK_NUMLOCK;		break;
+
+	case Windows::System::VirtualKey::Scroll:		TranslatedKeyCode = VK_SCROLL;		break;
+
+	case Windows::System::VirtualKey::LeftShift:	TranslatedKeyCode = VK_LSHIFT;		break;
+	case Windows::System::VirtualKey::RightShift:	TranslatedKeyCode = VK_RSHIFT;		break;
+	case Windows::System::VirtualKey::LeftControl:	TranslatedKeyCode = VK_LCONTROL;	break;
+	case Windows::System::VirtualKey::RightControl:	TranslatedKeyCode = VK_RCONTROL;	break;
+	case Windows::System::VirtualKey::LeftMenu:		TranslatedKeyCode = VK_LMENU;		break;
+	case Windows::System::VirtualKey::RightMenu:	TranslatedKeyCode = VK_RMENU;		break;
+
+	// For the US standard keyboard, only.
+#pragma warning(push)
+#pragma warning(disable:4063)
+	case 0xBA:										TranslatedKeyCode = VK_OEM_1;		break;
+	case 0xBB:										TranslatedKeyCode = VK_OEM_PLUS;	break;
+	case 0xBC:										TranslatedKeyCode = VK_OEM_COMMA;	break;
+	case 0xBD:										TranslatedKeyCode = VK_OEM_MINUS;	break;
+	case 0xBE:										TranslatedKeyCode = VK_OEM_PERIOD;	break;
+	case 0xBF:										TranslatedKeyCode = VK_OEM_2;		break;
+	case 0xC0:										TranslatedKeyCode = VK_OEM_3;		break;
+	case 0xDB:										TranslatedKeyCode = VK_OEM_4;		break;
+	case 0xDC:										TranslatedKeyCode = VK_OEM_5;		break;
+	case 0xDD:										TranslatedKeyCode = VK_OEM_6;		break;
+	case 0xDE:										TranslatedKeyCode = VK_OEM_7;		break;
+#pragma warning(pop)
+
+	default:
+		UE_LOG(LogLaunchHoloLens, Verbose, TEXT("Unrecognized keystroke VirtualKeyCode:%d  ScanCode:%d  Extended:%s  [%s]\n"),
+			TranslatedKeyCode,
+			keyEventArgs->KeyStatus.ScanCode,
+			keyEventArgs->KeyStatus.IsExtendedKey ? TEXT("yes") : TEXT("no"),
+			keyEventArgs->KeyStatus.WasKeyDown ? TEXT("up") : TEXT("down"));
+	}
+	return TranslatedKeyCode;
+}
+
+
+extern void appHoloLensKeyEvent(uint32 InCode, bool bInWasPressed);
+extern void appHoloLensCharEvent(uint32 InCode, bool bInWasPressed);
+
+void ViewProvider::OnKeyDown(
+	_In_ CoreWindow^ sender,
+	_In_ KeyEventArgs^ args 
+	)
+{
+	uint32 Code = TranslateWinRTKey(args);
+	if ((Code & 0xFFFFFF00) == 0)
+	{
+		appHoloLensKeyEvent(Code, true);
+	}
+	args->Handled = true;
+}
+void ViewProvider::OnKeyUp(
+	_In_ CoreWindow^ sender,
+	_In_ KeyEventArgs^ args 
+	)
+{
+	uint32 Code = TranslateWinRTKey(args);
+	if ((Code & 0xFFFFFF00) == 0)
+	{
+		appHoloLensKeyEvent(Code, false);
+	}
+	args->Handled = true;
+}
+
+void ViewProvider::OnAcceleratorKeyActivated(CoreDispatcher ^sender, AcceleratorKeyEventArgs ^args)
+{
+	if (args->VirtualKey == VirtualKey::Menu)
+	{
+		if (args->EventType == CoreAcceleratorKeyEventType::SystemKeyDown && args->KeyStatus.WasKeyDown == false)
+		{
+			appHoloLensKeyEvent(args->KeyStatus.IsExtendedKey ? VK_RMENU : VK_LMENU, true);
+		}
+		else if (args->EventType == CoreAcceleratorKeyEventType::SystemKeyUp && args->KeyStatus.WasKeyDown == true)
+		{
+			appHoloLensKeyEvent(args->KeyStatus.IsExtendedKey ? VK_RMENU : VK_LMENU, false);
+		}
+		args->Handled = true;
+	}
+
+	// Special case keycode for fullscreen
+	else if (args->VirtualKey == VirtualKey::Enter
+		&& args->EventType == CoreAcceleratorKeyEventType::SystemKeyDown
+		&& !args->KeyStatus.WasKeyDown)
+	{
+		FHoloLensApplication* HoloLensApplication = FHoloLensApplication::GetHoloLensApplication();
+		if (HoloLensApplication)
+		{
+			TSharedPtr< FHoloLensWindow > NativeWindow = HoloLensApplication->GetHoloLensWindow();
+			if (NativeWindow.IsValid())
+			{
+				if (NativeWindow->GetWindowMode() == EWindowMode::Windowed)
+				{
+					NativeWindow->SetWindowMode(EWindowMode::WindowedFullscreen);
+				}
+				else
+				{
+					NativeWindow->SetWindowMode(EWindowMode::Windowed);
+				}
+				args->Handled = true;
+			}
+		}
+	}
+	else if (args->VirtualKey == VirtualKey::Control)
+	{
+		if ((args->EventType == CoreAcceleratorKeyEventType::SystemKeyDown || args->EventType == CoreAcceleratorKeyEventType::KeyDown) && args->KeyStatus.WasKeyDown == false)
+		{
+			appHoloLensKeyEvent(args->KeyStatus.IsExtendedKey ? VK_RCONTROL : VK_LCONTROL, true);
+		}
+		else if ((args->EventType == CoreAcceleratorKeyEventType::SystemKeyUp || args->EventType == CoreAcceleratorKeyEventType::KeyUp) && args->KeyStatus.WasKeyDown == true)
+		{
+			appHoloLensKeyEvent(args->KeyStatus.IsExtendedKey ? VK_RCONTROL : VK_LCONTROL, false);
+		}
+		args->Handled = true;
+	}
+	else if (args->VirtualKey == VirtualKey::Shift)
+	{
+		if ((args->EventType == CoreAcceleratorKeyEventType::SystemKeyDown || args->EventType == CoreAcceleratorKeyEventType::KeyDown) && args->KeyStatus.WasKeyDown == false)
+		{
+			appHoloLensKeyEvent(args->KeyStatus.ScanCode == 0x36 ? VK_RSHIFT : VK_LSHIFT, true);
+		}
+		else if ((args->EventType == CoreAcceleratorKeyEventType::SystemKeyUp || args->EventType == CoreAcceleratorKeyEventType::KeyUp) && args->KeyStatus.WasKeyDown == true)
+		{
+			appHoloLensKeyEvent(args->KeyStatus.ScanCode == 0x36 ? VK_RSHIFT : VK_LSHIFT, false);
+		}
+		args->Handled = true;
+	}
+}
+
+void ViewProvider::OnCharacterReceived(
+	_In_ CoreWindow^ sender,
+	_In_ CharacterReceivedEventArgs^ args)
+{
+	if (args->KeyStatus.IsExtendedKey == false)
+	{
+		appHoloLensCharEvent(args->KeyCode, true);
+	}
+	args->Handled = true;
+}
+
+// @MIXEDREALITY_CHANGE : BEGIN Windows Mixed Reality support
+//bool bIsActivationContextHolographic = false;
+//
+//bool appHoloLensIsActivationHolographic()
+//{
+//	return bIsActivationContextHolographic;
+//}
+
+void appHoloLensViewProviderKeyEvent(
+	_In_ CoreWindow^ sender,
+	_In_ KeyEventArgs^ args,
+	bool down)
+{
+	uint32 Code = TranslateWinRTKey(args);
+	if ((Code & 0xFFFFFF00) == 0)
+	{
+		appHoloLensKeyEvent(Code, down);
+	}
+	args->Handled = true;
+}
+
+void appHoloLensViewProviderCharacterReceived(
+	_In_ CoreWindow^ sender,
+	_In_ CharacterReceivedEventArgs^ args)
+{
+	if (args->KeyStatus.IsExtendedKey == false)
+	{
+		appHoloLensCharEvent(args->KeyCode, true);
+	}
+	args->Handled = true;
+}
+
+void appHoloLensViewProviderAcceleratorKeyActivated(CoreDispatcher ^sender, AcceleratorKeyEventArgs ^args)
+{
+	if (args->VirtualKey == VirtualKey::Menu)
+	{
+		if (args->EventType == CoreAcceleratorKeyEventType::SystemKeyDown && args->KeyStatus.WasKeyDown == false)
+		{
+			appHoloLensKeyEvent(args->KeyStatus.IsExtendedKey ? VK_RMENU : VK_LMENU, true);
+		}
+		else if (args->EventType == CoreAcceleratorKeyEventType::SystemKeyUp && args->KeyStatus.WasKeyDown == true)
+		{
+			appHoloLensKeyEvent(args->KeyStatus.IsExtendedKey ? VK_RMENU : VK_LMENU, false);
+		}
+		args->Handled = true;
+	}
+
+	// Special case keycode for fullscreen
+	else if (args->VirtualKey == VirtualKey::Enter
+		&& args->EventType == CoreAcceleratorKeyEventType::SystemKeyDown
+		&& !args->KeyStatus.WasKeyDown)
+	{
+		FHoloLensApplication* HoloLensApplication = FHoloLensApplication::GetHoloLensApplication();
+		if (HoloLensApplication)
+		{
+			TSharedPtr< FHoloLensWindow > NativeWindow = HoloLensApplication->GetHoloLensWindow();
+			if (NativeWindow.IsValid())
+			{
+				if (NativeWindow->GetWindowMode() == EWindowMode::Windowed)
+				{
+					NativeWindow->SetWindowMode(EWindowMode::WindowedFullscreen);
+				}
+				else
+				{
+					NativeWindow->SetWindowMode(EWindowMode::Windowed);
+				}
+				args->Handled = true;
+			}
+		}
+	}
+	else if (args->VirtualKey == VirtualKey::Control)
+	{
+		if ((args->EventType == CoreAcceleratorKeyEventType::SystemKeyDown || args->EventType == CoreAcceleratorKeyEventType::KeyDown) && args->KeyStatus.WasKeyDown == false)
+		{
+			appHoloLensKeyEvent(args->KeyStatus.IsExtendedKey ? VK_RCONTROL : VK_LCONTROL, true);
+		}
+		else if ((args->EventType == CoreAcceleratorKeyEventType::SystemKeyUp || args->EventType == CoreAcceleratorKeyEventType::KeyUp) && args->KeyStatus.WasKeyDown == true)
+		{
+			appHoloLensKeyEvent(args->KeyStatus.IsExtendedKey ? VK_RCONTROL : VK_LCONTROL, false);
+		}
+		args->Handled = true;
+	}
+	else if (args->VirtualKey == VirtualKey::Shift)
+	{
+		if ((args->EventType == CoreAcceleratorKeyEventType::SystemKeyDown || args->EventType == CoreAcceleratorKeyEventType::KeyDown) && args->KeyStatus.WasKeyDown == false)
+		{
+			appHoloLensKeyEvent(args->KeyStatus.ScanCode == 0x36 ? VK_RSHIFT : VK_LSHIFT, true);
+		}
+		else if ((args->EventType == CoreAcceleratorKeyEventType::SystemKeyUp || args->EventType == CoreAcceleratorKeyEventType::KeyUp) && args->KeyStatus.WasKeyDown == true)
+		{
+			appHoloLensKeyEvent(args->KeyStatus.ScanCode == 0x36 ? VK_RSHIFT : VK_LSHIFT, false);
+		}
+		args->Handled = true;
+	}
+}
+// @MIXEDREALITY_CHANGE : END Windows Mixed Reality support
+
+void ViewProvider::OnVisibilityChanged(_In_ CoreWindow^ sender, _In_ VisibilityChangedEventArgs^ args)
+{
+	bVisible = args->Visible;
+}
+
+void ViewProvider::OnWindowSizeChanged(Windows::UI::Core::CoreWindow^ sender, Windows::UI::Core::WindowSizeChangedEventArgs^ args)
+{
+	float Dpi = static_cast<uint32_t>(Windows::Graphics::Display::DisplayInformation::GetForCurrentView()->LogicalDpi);
+	int32 Width = FHoloLensWindow::ConvertDipsToPixels(args->Size.Width, Dpi);
+	int32 Height = FHoloLensWindow::ConvertDipsToPixels(args->Size.Height, Dpi);
+
+	UE_LOG(LogLaunchHoloLens, Log, TEXT("Window Size Changed to [%d, %d]"), Width, Height);
+
+	FHoloLensApplication* pApplication = FHoloLensApplication::GetHoloLensApplication();
+	if (pApplication)
+	{
+		// Since we use the window size to define the virtual desktop size, we need to report display metrics changed here.
+		FDisplayMetrics DisplayMetrics;
+		FDisplayMetrics::RebuildDisplayMetrics(DisplayMetrics);
+		pApplication->OnDisplayMetricsChanged().Broadcast(DisplayMetrics);
+
+		TSharedPtr< FHoloLensWindow > pNativeWindow = pApplication->GetHoloLensWindow();
+		if (pNativeWindow.IsValid())
+		{
+			pApplication->GetMessageHandler()->OnResizingWindow(pNativeWindow.ToSharedRef());
+			pApplication->GetMessageHandler()->OnMovedWindow(pNativeWindow.ToSharedRef(), 0, 0);
+			pApplication->GetMessageHandler()->OnSizeChanged(pNativeWindow.ToSharedRef(), Width, Height);
+			pApplication->GetMessageHandler()->OnResizingWindow(pNativeWindow.ToSharedRef());
+		}
+	}
+	IConsoleManager::Get().CallAllConsoleVariableSinks();
+}
+
+void ViewProvider::OnWindowClosed(Windows::UI::Core::CoreWindow^ sender, Windows::UI::Core::CoreWindowEventArgs^ args)
+{
+	bWindowClosed = true;
+}
+
+void ViewProvider::OnPointerPressed(Windows::UI::Core::CoreWindow^ sender, Windows::UI::Core::PointerEventArgs^ args)
+{
+	ProcessMouseEvent(args);
+}
+
+void ViewProvider::OnPointerReleased(Windows::UI::Core::CoreWindow^ sender, Windows::UI::Core::PointerEventArgs^ args)
+{
+	ProcessMouseEvent(args); 
+}
+
+// ATG - This event occurs for any pointer updates, not just for movement, so must be certain to translate appropraitely
+void ViewProvider::OnPointerMoved(Windows::UI::Core::CoreWindow^ sender, Windows::UI::Core::PointerEventArgs^ args)
+{
+	ProcessMouseEvent(args);
+}
+
+void ViewProvider::OnPointerWheelChanged(Windows::UI::Core::CoreWindow^ sender, Windows::UI::Core::PointerEventArgs^ args)
+{
+	ProcessMouseEvent(args);
+}
+
+void ViewProvider::ProcessEvents()
+{
+	static std::atomic_flag Dispatching = ATOMIC_FLAG_INIT;
+
+	if (Dispatching.test_and_set( std::memory_order_acquire ))
+	{
+		return;		// prevent nesting
+	}
+
+	CoreWindow^ Window = CoreWindow::GetForCurrentThread();
+	if (nullptr != Window)
+	{
+		CoreDispatcher^ Dispatcher = Window->Dispatcher;
+
+		if (nullptr != Dispatcher)
+		{
+			// OpenXR uses a separate CoreWindow which will have visibility instead of the CoreWindow created by HoloLensLaunch and in this
+			// case ProcessOneAndAllPending cannot be used because it will block indefinitely once there are no more events to process.
+			if (bVisible || GSeparateCoreWindowVisible)
+			{
+				Dispatcher->ProcessEvents(CoreProcessEventsOption::ProcessAllIfPresent);
+			}
+			else
+			{
+				Dispatcher->ProcessEvents(CoreProcessEventsOption::ProcessOneAndAllPending);
+			}
+		}
+	}
+
+	TArray<QueuedPointerEvent> LocalPointerEvents(MoveTemp(PointerEventQueue));
+	PointerEventQueue.Empty();
+	Dispatching.clear(std::memory_order_release);
+
+	if (!IsEngineExitRequested())
+	{
+		for (const QueuedPointerEvent& Args : LocalPointerEvents)
+		{
+			ProcessMouseEvent(Args);
+		}
+	}
+}
+
+void appWinPumpMessages()
+{
+	GViewProvider->ProcessEvents();
+}
+
+void ViewProvider::Initialize(Windows::ApplicationModel::Core::CoreApplicationView^ applicationView)
+{
+	applicationView->Activated += ref new Windows::Foundation::TypedEventHandler< CoreApplicationView^, IActivatedEventArgs^ >( this, &ViewProvider::OnActivated );
+	CoreApplication::Suspending += ref new Windows::Foundation::EventHandler< Windows::ApplicationModel::SuspendingEventArgs^ >( this, &ViewProvider::OnSuspending );
+	CoreApplication::Resuming += ref new Windows::Foundation::EventHandler< Platform::Object^>( this, & ViewProvider::OnResuming );
+	CoreApplication::EnteredBackground += ref new Windows::Foundation::EventHandler< Windows::ApplicationModel::EnteredBackgroundEventArgs^ >(this, &ViewProvider::OnEnteredBackground);
+	CoreApplication::LeavingBackground += ref new Windows::Foundation::EventHandler< Windows::ApplicationModel::LeavingBackgroundEventArgs^ >(this, &ViewProvider::OnLeavingBackground);
+	CoreApplication::Exiting += ref new Windows::Foundation::EventHandler< Platform::Object^>(this, &ViewProvider::OnExiting);
+
+}
+
+static void InitCommandLine(Platform::String^ args)
+{
+	// check for a uecommandline.txt.
+	FCommandLine::Set(L"");
+
+	// FPaths should be good at this point (FPaths::EngineDir() is used during static init for instance)
+	// so be specific about the location we want, otherwise working directory changes can mess this up.
+	// Also this means the regular file interface is in place, so let's use that.
+	FString FileName = FString(FPaths::RootDir()) / TEXT("UECommandLine.txt");
+
+	IFileHandle* CmdLineFileHandle = IPlatformFile::GetPlatformPhysical().OpenRead(*FileName);
+	if (CmdLineFileHandle != nullptr)
+	{
+		const int64 MaxCmdLineSize = 65536;
+
+		int64 CmdLineFileSize = CmdLineFileHandle->Size();
+		if (CmdLineFileSize < MaxCmdLineSize - 1)
+		{
+			char AnsiCmdLine[MaxCmdLineSize] = {};
+			if (CmdLineFileHandle->Read(reinterpret_cast<uint8*>(AnsiCmdLine), CmdLineFileSize))
+			{
+				FString CmdLine = StringCast<TCHAR>(AnsiCmdLine).Get();
+				CmdLine.TrimStartAndEndInline();
+				
+				if (args != nullptr)
+				{
+					CmdLine += FString(" ");
+					CmdLine += FString(args->Data());
+
+					CmdLine.TrimStartAndEndInline();
+				}
+
+				UE_LOG(LogLaunchHoloLens, Log, TEXT("%s"), *CmdLine);
+				FCommandLine::Set(*CmdLine);
+			}
+			else
+			{
+				UE_LOG(LogLaunchHoloLens, Warning, TEXT("Failed to read commandline from %s!"), *FileName);
+			}
+		}
+		else
+		{
+			UE_LOG(LogLaunchHoloLens, Warning, TEXT("Commandline file %s too large (%d / %d bytes).  Ignoring."), *FileName, CmdLineFileSize, MaxCmdLineSize);
+		}
+		delete CmdLineFileHandle;
+	}
+}
+
+void ViewProvider::OnActivated(_In_ Windows::ApplicationModel::Core::CoreApplicationView^ applicationView, _In_ Windows::ApplicationModel::Activation::IActivatedEventArgs^ args)
+{
+	// Check for launch activation
+	Platform::String^ arg = nullptr;
+	if (args->Kind == ActivationKind::Launch)
+	{
+		auto launchArgs = static_cast<LaunchActivatedEventArgs^>(args);
+		arg = launchArgs->Arguments;
+	}
+
+	// Check for protocol activation
+	if (args->Kind == ActivationKind::Protocol)
+	{
+		auto protocolArgs = static_cast<ProtocolActivatedEventArgs^>(args);
+		arg = protocolArgs->Uri->ToString();
+	}
+
+	InitCommandLine(arg);
+	
+	// Check commandline for start in VR.
+	if (FParse::Param(FCommandLine::Get(), TEXT("vr")))
+	{
+		auto HoloSpace = HolographicSpace::CreateForCoreWindow(CoreWindow::GetForCurrentThread());
+		FHoloLensApplication::SetHolographicSpace(HoloSpace);
+	}
+	else
+	{
+		UE_LOG(LogHoloLens, Warning, TEXT("Application will launch as a 2D slate.  If this is unexpected, add a -vr command line flag or enable the start in VR flag."));
+	}
+
+	// Take advantage of this opportunity to measure the desktop (based on view bounds before window activation).
+	FHoloLensApplication::CacheDesktopSize();
+
+	// Activate the window
+	CoreWindow^ window = CoreWindow::GetForCurrentThread();
+	window->Activate();
+    ActivationComplete = true;
+
+	// Query this to find out if the application was shut down gracefully last time
+	if (args != nullptr)
+	{
+		if(args->Kind == ActivationKind::Launch)
+		{
+			LaunchActivatedEventArgs ^LaunchArgs = (LaunchActivatedEventArgs^)args;
+			if (!LaunchArgs->Arguments->IsEmpty())
+			{
+				FCommandLine::Set(LaunchArgs->Arguments->Data());
+			}
+		}
+		else if(args->Kind == Windows::ApplicationModel::Activation::ActivationKind::Protocol)
+		{
+			// If the game is launched from an invitation, the ActivationKind will be Protocol
+			// so we need to initialize the command line on this path also. Since it's possible
+			// to get a protocol activation after the game is already running as well, only set
+			// the command line if it's not initialized yet.
+			if ( !FCommandLine::IsInitialized() )
+			{
+				FCommandLine::Set( TEXT("") );
+			}
+
+			// Also save the actual protocol activation URI, for other parts of the engine to
+			// access if necessary once they're initialized.
+			// This may overwrite an existing URI if the game is activated multiple times,
+			// but we assume only the most recent activation is relevant.
+			ProtocolActivatedEventArgs^ ProtocolArgs = static_cast<ProtocolActivatedEventArgs^>(args);
+			FPlatformMisc::SetProtocolActivationUri(ProtocolArgs->Uri->RawUri->Data());
+		}
+
+		existingSize = window->Bounds;
+		ApplicationExecutionState lastState = args->PreviousExecutionState;
+		UE_LOG(LogLaunchHoloLens, Log, TEXT("OnActivated (%5.0fx%5.0f): Last application state: %d"), existingSize.Width, existingSize.Height, static_cast<uint32>(lastState));
+	}
+}
+
+void ViewProvider::OnResuming(_In_ Platform::Object^ Sender, _In_ Platform::Object^ Args)
+{
+	// Make the call down to the RHI to Resume the GPU state
+	RHIResumeRendering();
+
+	// Notify application of resume
+	FCoreDelegates::ApplicationHasReactivatedDelegate.Broadcast();
+}
+
+void ViewProvider::OnSuspending(_In_ Platform::Object^ Sender, _In_ Windows::ApplicationModel::SuspendingEventArgs^ Args)
+{
+	// Get the Suspending Event
+	Windows::ApplicationModel::SuspendingDeferral^ SuspendingEvent = Args->SuspendingOperation->GetDeferral();
+
+	// Notify application of suspend. Application should kick off an async save at this point.
+	FCoreDelegates::ApplicationWillDeactivateDelegate.Broadcast();
+
+	// Flush the RenderingThread
+	FlushRenderingCommands();
+
+	// Make the call down to the RHI to Suspend the GPU state
+	if (GDynamicRHI != nullptr)
+	{
+		RHISuspendRendering();
+	}
+
+	// @TODO Wait for async save to complete
+	// Flush the log so it's all written to disk
+	GLog->FlushThreadedLogs();
+	GLog->Flush();
+
+	// Tell the callback that we are done
+	SuspendingEvent->Complete();
+}
+
+void ViewProvider::OnEnteredBackground(_In_ Platform::Object^ sender, _In_ Windows::ApplicationModel::EnteredBackgroundEventArgs^ args)
+{
+	if (!GIsRunning)
+	{
+		return;
+	}
+
+	FCoreDelegates::ApplicationWillEnterBackgroundDelegate.Broadcast();
+}
+
+void ViewProvider::OnLeavingBackground(_In_ Platform::Object^ sender, _In_ Windows::ApplicationModel::LeavingBackgroundEventArgs^ args)
+{
+	if (!GIsRunning)
+	{
+		return;
+	}
+
+	FCoreDelegates::ApplicationHasEnteredForegroundDelegate.Broadcast();
+}
+
+void ViewProvider::OnExiting(_In_ Platform::Object^ sender, _In_ Platform::Object^ args)
+{
+	FCoreDelegates::ApplicationWillTerminateDelegate.Broadcast();
+}
+
+
+const TCHAR* ViewProvider::GetPointerUpdateKindString(Windows::UI::Input::PointerUpdateKind InKind)
+{
+	switch (InKind)
+	{
+	case Windows::UI::Input::PointerUpdateKind::Other:
+		return TEXT("Other");
+	case Windows::UI::Input::PointerUpdateKind::LeftButtonPressed:
+		return TEXT("LeftButtonPressed");
+	case Windows::UI::Input::PointerUpdateKind::LeftButtonReleased:
+		return TEXT("LeftButtonReleased");
+	case Windows::UI::Input::PointerUpdateKind::RightButtonPressed:
+		return TEXT("RightButtonPressed");
+	case Windows::UI::Input::PointerUpdateKind::RightButtonReleased:
+		return TEXT("RightButtonReleased");
+	case Windows::UI::Input::PointerUpdateKind::MiddleButtonPressed:
+		return TEXT("MiddleButtonPressed");
+	case Windows::UI::Input::PointerUpdateKind::MiddleButtonReleased:
+		return TEXT("MiddleButtonReleased");
+	case Windows::UI::Input::PointerUpdateKind::XButton1Pressed:
+		return TEXT("XButton1Pressed");
+	case Windows::UI::Input::PointerUpdateKind::XButton1Released:
+		return TEXT("XButton1Released");
+	case Windows::UI::Input::PointerUpdateKind::XButton2Pressed:
+		return TEXT("XButton2Pressed");
+	case Windows::UI::Input::PointerUpdateKind::XButton2Released:
+		return TEXT("XButton2Released");
+	}
+
+	return TEXT("*** UNKNOWN ***");
+}
+
+
+EMouseButtons::Type ViewProvider::PointerUpdateKindToUEKey(Windows::UI::Input::PointerUpdateKind InKind, bool& bWasPressed)
+{
+	switch (InKind)
+	{
+	case Windows::UI::Input::PointerUpdateKind::Other:
+		bWasPressed = false; ////
+		return EMouseButtons::Invalid; ////
+	case Windows::UI::Input::PointerUpdateKind::LeftButtonPressed:
+		bWasPressed = true;
+		return EMouseButtons::Left;
+	case Windows::UI::Input::PointerUpdateKind::LeftButtonReleased:
+		bWasPressed = false;
+		return EMouseButtons::Left;
+	case Windows::UI::Input::PointerUpdateKind::RightButtonPressed:
+		bWasPressed = true;
+		return EMouseButtons::Right;
+	case Windows::UI::Input::PointerUpdateKind::RightButtonReleased:
+		bWasPressed = false;
+		return EMouseButtons::Right;
+	case Windows::UI::Input::PointerUpdateKind::MiddleButtonPressed:
+		bWasPressed = true;
+		return EMouseButtons::Middle;
+	case Windows::UI::Input::PointerUpdateKind::MiddleButtonReleased:
+		bWasPressed = false;
+		return EMouseButtons::Middle;
+	case Windows::UI::Input::PointerUpdateKind::XButton1Pressed:
+		bWasPressed = true;
+		return EMouseButtons::Thumb01;
+	case Windows::UI::Input::PointerUpdateKind::XButton1Released:
+		bWasPressed = false;
+		return EMouseButtons::Thumb01;
+	case Windows::UI::Input::PointerUpdateKind::XButton2Pressed:
+		bWasPressed = true;
+		return EMouseButtons::Thumb02;
+	case Windows::UI::Input::PointerUpdateKind::XButton2Released:
+		bWasPressed = false;
+		return EMouseButtons::Thumb02;
+	}
+
+	return EMouseButtons::Invalid; ////
+}
+
+bool ViewProvider::ProcessMouseEvent(Windows::UI::Core::PointerEventArgs^ args)
+{
+	Windows::UI::Input::PointerPoint^ Point = args->CurrentPoint;
+	Windows::UI::Input::PointerUpdateKind Kind = Point->Properties->PointerUpdateKind;
+	Windows::System::VirtualKeyModifiers KeyModifiers = args->KeyModifiers;
+	Windows::Foundation::Collections::IVector<Windows::UI::Input::PointerPoint^>^ intermediatePoints = args->GetIntermediatePoints();
+
+	UE_LOG(LogLaunchHoloLens, Verbose, TEXT("ProcessMouseEvent %d = %5.2f, %5.2f - %s"),
+		intermediatePoints->Size, Point->Position.X, Point->Position.Y, GetPointerUpdateKindString(Kind));
+
+	PointerEventQueue.Add(QueuedPointerEvent(args));
+
+	return true;
+}
+
+bool ViewProvider::ProcessMouseEvent(const QueuedPointerEvent& Event)
+{
+	FHoloLensApplication* const Application = FHoloLensApplication::GetHoloLensApplication();
+
+	if (Application != NULL)
+	{
+		// process a cursor move unless we're using raw mouse deltas and a virtual cursor
+		if (!Application->GetCursor()->IsUsingRawMouseNoCursor())
+		{
+			FVector2D CurrentCursorPosition = Event.RawPosition;
+			float Dpi = static_cast<uint32_t>(Windows::Graphics::Display::DisplayInformation::GetForCurrentView()->LogicalDpi);
+			CurrentCursorPosition.X = FHoloLensWindow::ConvertDipsToPixels(CurrentCursorPosition.X, Dpi);
+			CurrentCursorPosition.Y = FHoloLensWindow::ConvertDipsToPixels(CurrentCursorPosition.Y, Dpi);
+
+			Application->GetCursor()->UpdatePosition(CurrentCursorPosition);
+			Application->GetMessageHandler()->OnMouseMove();
+			Application->GetMessageHandler()->OnCursorSet();
+		}
+		// process a button event if we know this update isn't just for movement
+		if (Event.Kind != Windows::UI::Input::PointerUpdateKind::Other)
+		{
+			bool bPressed = false;
+			EMouseButtons::Type MouseButton = PointerUpdateKindToUEKey(Event.Kind, bPressed);
+			if (MouseButton != EMouseButtons::Type::Invalid)
+			{
+				if (bPressed == true)
+				{
+					Application->GetMessageHandler()->OnMouseDown(NULL, MouseButton);
+				}
+				else
+				{
+					Application->GetMessageHandler()->OnMouseUp(MouseButton);
+				}
+			}
+		}
+
+		if (Event.WheelDelta != 0)
+		{
+			Application->GetMessageHandler()->OnMouseWheel(Event.WheelDelta / static_cast<float>(WHEEL_DELTA));
+		}
+	}
+
+	return true;
+}
+
+
+
+void ViewProvider::SetWindow(Windows::UI::Core::CoreWindow^ window)
+{
+	FModuleManager::Get().LoadModule(TEXT("WindowsMixedRealityRHI"));
+
+	window->KeyDown += ref new Windows::Foundation::TypedEventHandler<CoreWindow^, KeyEventArgs^>(this, &ViewProvider::OnKeyDown);
+	window->KeyUp += ref new Windows::Foundation::TypedEventHandler<CoreWindow^, KeyEventArgs^>(this, &ViewProvider::OnKeyUp);
+	window->CharacterReceived += ref new Windows::Foundation::TypedEventHandler<CoreWindow^, CharacterReceivedEventArgs^>(this, &ViewProvider::OnCharacterReceived);
+	window->VisibilityChanged += ref new TypedEventHandler<CoreWindow^, VisibilityChangedEventArgs^>(this, &ViewProvider::OnVisibilityChanged);
+	window->SizeChanged += ref new TypedEventHandler<CoreWindow^, WindowSizeChangedEventArgs^>(this, &ViewProvider::OnWindowSizeChanged);
+	window->Closed += ref new TypedEventHandler<CoreWindow^, CoreWindowEventArgs^>(this, &ViewProvider::OnWindowClosed);
+	window->PointerCursor = ref new CoreCursor(CoreCursorType::Arrow, 0);
+	window->PointerPressed += ref new TypedEventHandler<CoreWindow^, PointerEventArgs^>(this, &ViewProvider::OnPointerPressed);
+	window->PointerReleased += ref new TypedEventHandler<CoreWindow^, PointerEventArgs^>(this, &ViewProvider::OnPointerReleased);
+	window->PointerMoved += ref new TypedEventHandler<CoreWindow^, PointerEventArgs^>(this, &ViewProvider::OnPointerMoved);
+	window->PointerWheelChanged += ref new TypedEventHandler<CoreWindow^, PointerEventArgs^>(this, &ViewProvider::OnPointerWheelChanged);
+	window->Dispatcher->AcceleratorKeyActivated += ref new TypedEventHandler<CoreDispatcher ^, AcceleratorKeyEventArgs ^>(this, &ViewProvider::OnAcceleratorKeyActivated);
+}
+
+// this method is called after Initialize
+void ViewProvider::Load(Platform::String^ /*entryPoint*/)
+{
+}
+
+void ViewProvider::EnsureDlcPluginsAreMounted()
+{
+	IPluginManager::Get().RefreshPluginsList();
+	TArray<TSharedRef<IPlugin>> DiscoveredPlugins = IPluginManager::Get().GetDiscoveredPlugins();
+	for (TSharedRef<IPlugin>& Plugin : DiscoveredPlugins)
+	{
+		if (Plugin->GetType() == EPluginType::External && !Plugin->IsEnabled())
+		{
+			IPluginManager::Get().MountNewlyCreatedPlugin(Plugin->GetName());
+		}
+	}
+}
+
+void ViewProvider::RegisterDlcPluginsForOptionalPackage(Windows::ApplicationModel::Package^ DependencyPackage)
+{
+#if WIN10_SDK_VERSION >= 14393
+	check(Windows::Foundation::Metadata::ApiInformation::IsPropertyPresent("Windows.ApplicationModel.Package", "IsOptional"));
+	// Optional dependencies are the mechanism for DLC on HoloLens
+	if (DependencyPackage->IsOptional)
+	{
+		bool ShouldLoad = true;
+		if (DependencyPackage->Status->NeedsRemediation)
+		{
+			UE_LOG(LogLaunchHoloLens, Warning, TEXT("Optional package %s has NeedsRemediation set and will not be loaded."), DependencyPackage->Id->FullName->Data());
+			ShouldLoad = false;
+		}
+		else if (DependencyPackage->Status->NotAvailable)
+		{
+			UE_LOG(LogLaunchHoloLens, Warning, TEXT("Optional package %s has NotAvailable set and will not be loaded."), DependencyPackage->Id->FullName->Data());
+			ShouldLoad = false;
+		}
+
+		if (ShouldLoad)
+		{
+			IPluginManager::Get().AddPluginSearchPath(FString(DependencyPackage->InstalledLocation->Path->Data()), false);
+		}
+	}
+#endif
+}
+
+void ViewProvider::InitOptionalPackages()
+{
+#if WIN10_SDK_VERSION >= 14393
+	if (Windows::Foundation::Metadata::ApiInformation::IsTypePresent("Windows.ApplicationModel.PackageCatalog"))
+	{
+		using namespace Windows::Foundation;
+		using namespace Windows::Foundation::Collections;
+		using namespace Windows::ApplicationModel;
+
+		IVectorView<Package^>^ DependencyPackages = Package::Current->Dependencies;
+		for (Package^ DependencyPackage : DependencyPackages)
+		{
+			RegisterDlcPluginsForOptionalPackage(DependencyPackage);
+		}
+		EnsureDlcPluginsAreMounted();
+
+		static PackageCatalog^ Catalog = PackageCatalog::OpenForCurrentPackage();
+		Catalog->PackageInstalling += ref new TypedEventHandler<PackageCatalog^, PackageInstallingEventArgs^>(
+			[](PackageCatalog^, PackageInstallingEventArgs^ Args)
+		{
+			if (Args->IsComplete)
+			{
+				RegisterDlcPluginsForOptionalPackage(Args->Package);
+				EnsureDlcPluginsAreMounted();
+			}
+		});
+
+		Catalog->PackageStatusChanged += ref new TypedEventHandler<PackageCatalog^, PackageStatusChangedEventArgs^>(
+			[](PackageCatalog^, PackageStatusChangedEventArgs^ Args)
+		{
+			RegisterDlcPluginsForOptionalPackage(Args->Package);
+			EnsureDlcPluginsAreMounted();
+		});
+	}
+	else
+	{
+		UE_LOG(LogLaunchHoloLens, Warning, TEXT("Optional packages (DLC) require a runtime Windows version of 14393 or higher."));
+	}
+#else
+	UE_LOG(LogLaunchHoloLens, Warning, TEXT("Optional packages (DLC) are not available with this Windows SDK version (%d).  Use the 14393 SDK or later."), WIN10_SDK_VERSION);
+#endif
+}
+
+#include "Microsoft/AllowMicrosoftPlatformAtomics.h"
+#include "Microsoft/AllowMicrosoftPlatformTypes.h"
+#include "NetworkMessage.h"
+
+// this method is called after Load
+void ViewProvider::Run()
+{
+	FTaskTagScope Scope(ETaskTag::EGameThread);
+
+	// Wait for activation to complete
+	while (!ActivationComplete)
+	{
+		ProcessEvents();
+	}
+
+	appHoloLensEarlyInit();
+
+	// command line fetch moved to earlier in startup, in a background task
+
+	// Parse the basedir and check for the 'clearbasedir' option
+	const TCHAR* BaseDirectory = FPlatformProcess::BaseDir();
+	if (FParse::Param(FCommandLine::Get(), TEXT("clearbasedir")))
+	{
+		// An actual path was passed in... check it
+		WIN32_FILE_ATTRIBUTE_DATA AttribData;
+		if (GetFileAttributesEx(BaseDirectory, GetFileExInfoStandard, (void*)&AttribData) == TRUE)
+		{
+			if ((AttribData.dwFileAttributes & FILE_ATTRIBUTE_DIRECTORY) != 0)
+			{
+				// Deleting the directory does not appear to work.
+				// So delete the files one by one
+				UE_LOG(LogLaunchHoloLens, Warning, TEXT("WARNING: Requested to clear base directory: %s"), BaseDirectory);
+				IPlatformFile& PlatformFile = FPlatformFileManager::Get().GetPlatformFile();
+				TArray<FString> DirectoriesToSkip;
+				TArray<FString> DirectoriesToNotRecurse;
+				FLocalTimestampDirectoryVisitor Visitor(PlatformFile, DirectoriesToSkip, DirectoriesToNotRecurse, false);
+				PlatformFile.IterateDirectory(BaseDirectory, Visitor);
+				if (Visitor.FileTimes.Num() > 0)
+				{
+					for (TMap<FString,FDateTime>::TIterator FileIt(Visitor.FileTimes); FileIt; ++FileIt)
+					{
+						DeleteFileW(*(FileIt.Key()));
+					}
+				}
+			}
+		}
+	}
+
+	GuardedMain(FCommandLine::Get(), NULL, NULL, 0 );
+
+	Windows::ApplicationModel::Core::CoreApplication::Exit();
+}
+#include "Microsoft/HideMicrosoftPlatformTypes.h"
+#include "Microsoft/HideMicrosoftPlatformAtomics.h"
+
+void ViewProvider::Uninitialize()
+{
+}
+
+[Platform::MTAThread]
+int main(Platform::Array<Platform::String^>^)
+{
+	uint64 GameThreadAffinity = FPlatformAffinity::GetMainGameMask();
+	FPlatformProcess::SetThreadAffinityMask((DWORD_PTR)GameThreadAffinity);
+
+	auto viewProviderFactory = ref new ViewProviderFactory();
+	Windows::ApplicationModel::Core::CoreApplication::Run(viewProviderFactory);
+}
+
+// The following functions are also located in Launch.cpp in a #if _WINDOWS block...
+/**
+ * Handler for CRT parameter validation. Triggers error
+ *
+ * @param Expression - the expression that failed crt validation
+ * @param Function - function which failed crt validation
+ * @param File - file where failure occured
+ * @param Line - line number of failure
+ * @param Reserved - not used
+ */
+void InvalidParameterHandler(const TCHAR* Expression,
+							 const TCHAR* Function, 
+							 const TCHAR* File, 
+							 uint32 Line, 
+							 uintptr_t Reserved)
+{
+	UE_LOG(LogLaunchHoloLens, Fatal,
+		TEXT("SECURE CRT: Invalid parameter detected.\nExpression: %s Function: %s. File: %s Line: %d\n"), 
+		Expression ? Expression : TEXT("Unknown"), 
+		Function ? Function : TEXT("Unknown"), 
+		File ? File : TEXT("Unknown"), 
+		Line );
+}
+
+/**
+ * Setup the common debug settings 
+ */
+void SetupWindowsEnvironment( void )
+{
+	// all crt validation should trigger the callback
+	_set_invalid_parameter_handler(InvalidParameterHandler);
+}
+
+void appHoloLensEarlyInit()
+{
+	// Setup common Windows settings
+	SetupWindowsEnvironment();
+
+	int32 ErrorLevel			= 0;
+}
+
+/** The global EngineLoop instance */
+FEngineLoop	GEngineLoop;
+bool GIsConsoleExecutable = false;
+
+/** 
+ * PreInits the engine loop 
+ */
+int32 EnginePreInit( const TCHAR* CmdLine )
+{
+	int32 ErrorLevel = GEngineLoop.PreInit(CmdLine);
+
+	return(ErrorLevel);
+}
+
+/** 
+ * Inits the engine loop 
+ */
+int32 EngineInit( void )
+{
+	GViewProvider->InitOptionalPackages();
+	return GEngineLoop.Init();
+}
+
+/** 
+ * Ticks the engine loop 
+ */
+void EngineTick( void )
+{
+	GEngineLoop.Tick();
+}
+
+/**
+ * Shuts down the engine
+ */
+void EngineExit( void )
+{
+	// Make sure this is set
+	RequestEngineExit(TEXT("EngineExit() was called"));
+	GEngineLoop.Exit();
+}
+
+/**
+ * Static guarded main function. Rolled into own function so we can have error handling for debug/ release builds depending
+ * on whether a debugger is attached or not.
+ */
+int32 GuardedMain( const TCHAR* CmdLine, HINSTANCE hInInstance, HINSTANCE hPrevInstance, int32 nCmdShow )
+{
+#if UE_BUILD_DEBUG
+	UE_LOG(LogLaunchHoloLens, Log, TEXT("--- Running a DEBUG build ---"));
+#elif UE_BUILD_DEVELOPMENT
+	UE_LOG(LogLaunchHoloLens, Log, TEXT("--- Running a DEVELOPMENT build ---"));
+#elif UE_BUILD_TEST
+	UE_LOG(LogLaunchHoloLens, Log, TEXT("--- Running a TEST build ---"));
+#elif UE_BUILD_SHIPPING
+	UE_LOG(LogLaunchHoloLens, Log, TEXT("--- Running a SHIPPING build ---"));
+#else
+	UE_LOG(LogLaunchHoloLens, Log, TEXT("--- Running a ***UNKNOWN*** build ---"));
+#endif
+
+	// make sure GEngineLoop::Exit() is always called.
+	struct EngineLoopCleanupGuard 
+	{ 
+		~EngineLoopCleanupGuard()
+		{
+			EngineExit();
+		}
+	} CleanupGuard;
+
+	// Set up minidump filename. We cannot do this directly inside main as we use an FString that requires 
+	// destruction and main uses SEH.
+	// These names will be updated as soon as the Filemanager is set up so we can write to the log file.
+	// That will also use the user folder for installed builds so we don't write into program files or whatever.
+	FCString::Strcpy(MiniDumpFilenameW, *FString::Printf( TEXT("unreal-v%i-%s.dmp"), FEngineVersion::Current().GetChangelist(), *FDateTime::Now().ToString()));
+
+	GViewProvider->ProcessEvents();
+
+	// Call PreInit and exit if failed.
+	int32 ErrorLevel = EnginePreInit(CmdLine);
+	if ((ErrorLevel != 0) || IsEngineExitRequested())
+	{
+		return ErrorLevel;
+	}
+
+	GViewProvider->ProcessEvents();
+
+	// Game without wxWindows.
+	ErrorLevel = EngineInit();
+
+	static const FName OpenXRSystemName(TEXT("OpenXR"));
+	GSeparateCoreWindowVisible = GEngine->XRSystem.IsValid() && (GEngine->XRSystem->GetSystemName() == OpenXRSystemName);
+
+	uint64 GameThreadAffinity = FPlatformAffinity::GetMainGameMask();
+	FPlatformProcess::SetThreadAffinityMask((DWORD_PTR)GameThreadAffinity);
+
+	while (!IsEngineExitRequested())
+	{
+		// ProcessEvents will throw if the process is exiting, allowing us to
+		// break out of the loop.  This will be cleaned up when we get proper
+		// process lifetime management online in a future release.
+		GViewProvider->ProcessEvents();
+        FHoloLensApplication::GetHoloLensApplication()->GetCursor()->ProcessDeferredActions();
+
+		EngineTick();
+	}
+	return ErrorLevel;
+}
diff --git a/Engine/Platforms/Hololens/Source/Runtime/Online/BackgroundHTTP/BackgroundHTTP_HoloLens.build.cs b/Engine/Platforms/Hololens/Source/Runtime/Online/BackgroundHTTP/BackgroundHTTP_HoloLens.build.cs
new file mode 100644
--- /dev/null
+++ b/Engine/Platforms/Hololens/Source/Runtime/Online/BackgroundHTTP/BackgroundHTTP_HoloLens.build.cs
@@ -0,0 +1,11 @@
+// Copyright Epic Games, Inc. All Rights Reserved.
+
+namespace UnrealBuildTool.Rules
+{
+	public class BackgroundHTTP_HoloLens : BackgroundHTTP
+	{
+		public BackgroundHTTP_HoloLens(ReadOnlyTargetRules Target) : base(Target)
+		{
+		}
+	}
+}
diff --git a/Engine/Platforms/Hololens/Source/Runtime/Online/BackgroundHTTP/Public/HoloLensPlatformBackgroundHttp.h b/Engine/Platforms/Hololens/Source/Runtime/Online/BackgroundHTTP/Public/HoloLensPlatformBackgroundHttp.h
new file mode 100644
--- /dev/null
+++ b/Engine/Platforms/Hololens/Source/Runtime/Online/BackgroundHTTP/Public/HoloLensPlatformBackgroundHttp.h
@@ -0,0 +1,8 @@
+// Copyright Epic Games, Inc. All Rights Reserved.
+
+#pragma once
+
+#include "GenericPlatform/GenericPlatformBackgroundHttp.h"
+
+//Currently no implementation on PS4
+typedef FGenericPlatformBackgroundHttp FPlatformBackgroundHttp;
\ No newline at end of file
diff --git a/Engine/Platforms/Hololens/Source/Runtime/Online/HTTP/HTTP_HoloLens.build.cs b/Engine/Platforms/Hololens/Source/Runtime/Online/HTTP/HTTP_HoloLens.build.cs
new file mode 100644
--- /dev/null
+++ b/Engine/Platforms/Hololens/Source/Runtime/Online/HTTP/HTTP_HoloLens.build.cs
@@ -0,0 +1,11 @@
+// Copyright Epic Games, Inc. All Rights Reserved.
+
+namespace UnrealBuildTool.Rules
+{
+	public class HTTP_HoloLens : HTTP
+	{
+		public HTTP_HoloLens(ReadOnlyTargetRules Target) : base(Target)
+		{
+		}
+	}
+}
diff --git a/Engine/Platforms/Hololens/Source/Runtime/Online/HTTP/Private/HoloLensPlatformHttp.cpp b/Engine/Platforms/Hololens/Source/Runtime/Online/HTTP/Private/HoloLensPlatformHttp.cpp
new file mode 100644
--- /dev/null
+++ b/Engine/Platforms/Hololens/Source/Runtime/Online/HTTP/Private/HoloLensPlatformHttp.cpp
@@ -0,0 +1,22 @@
+// Copyright Epic Games, Inc. All Rights Reserved.
+
+#include "HoloLensPlatformHttp.h"
+#include "IXML/HttpIXML.h"
+
+void FHoloLensHttp::Init()
+{
+}
+
+void FHoloLensHttp::Shutdown()
+{
+}
+
+FHttpManager * FHoloLensHttp::CreatePlatformHttpManager()
+{
+	return nullptr;
+}
+
+IHttpRequest* FHoloLensHttp::ConstructRequest()
+{
+	return new FHttpRequestIXML();
+}
\ No newline at end of file
diff --git a/Engine/Platforms/Hololens/Source/Runtime/Online/ICMP/ICMP_HoloLens.build.cs b/Engine/Platforms/Hololens/Source/Runtime/Online/ICMP/ICMP_HoloLens.build.cs
new file mode 100644
--- /dev/null
+++ b/Engine/Platforms/Hololens/Source/Runtime/Online/ICMP/ICMP_HoloLens.build.cs
@@ -0,0 +1,11 @@
+// Copyright Epic Games, Inc. All Rights Reserved.
+
+namespace UnrealBuildTool.Rules
+{
+	public class Icmp_HoloLens : Icmp
+	{
+		public Icmp_HoloLens(ReadOnlyTargetRules Target) : base(Target)
+		{
+		}
+	}
+}
diff --git a/Engine/Platforms/Hololens/Source/Runtime/Online/WebSockets/Private/HoloLensWebSocket.cpp b/Engine/Platforms/Hololens/Source/Runtime/Online/WebSockets/Private/HoloLensWebSocket.cpp
new file mode 100644
--- /dev/null
+++ b/Engine/Platforms/Hololens/Source/Runtime/Online/WebSockets/Private/HoloLensWebSocket.cpp
@@ -0,0 +1,231 @@
+// Copyright Epic Games, Inc. All Rights Reserved.
+
+#include "HoloLensWebSocket.h"
+#include "Stats/Stats.h"
+
+#if WITH_WEBSOCKETS
+
+#include "WebSocketsLog.h"
+
+FHoloLensWebSocket::FHoloLensWebSocket(const FString& InUri, const TArray<FString>& InProtocols, const TMap<FString, FString>& InHeaders)
+	: Uri(InUri)
+	, Socket(ref new Windows::Networking::Sockets::MessageWebSocket())
+	, bUserClose(false)
+{
+	Windows::Networking::Sockets::MessageWebSocketControl^ SocketControl = Socket->Control;
+	SocketControl->MessageType = Windows::Networking::Sockets::SocketMessageType::Utf8;
+	for (const FString& Protocol : InProtocols)
+	{
+		SocketControl->SupportedProtocols->Append(ref new Platform::String(*Protocol));
+	}
+
+	for (TMap<FString, FString>::TConstIterator It(InHeaders); It; ++It)
+	{
+		Socket->SetRequestHeader(ref new Platform::String(*It->Key), ref new Platform::String(*It->Value));
+	}
+}
+
+void FHoloLensWebSocket::Connect()
+{
+	using namespace Windows::Networking::Sockets;
+
+	TWeakPtr<FHoloLensWebSocket> WeakThis = SharedThis(this);
+	Socket->MessageReceived += ref new Windows::Foundation::TypedEventHandler<MessageWebSocket^, MessageWebSocketMessageReceivedEventArgs^>(
+		[WeakThis](MessageWebSocket^, MessageWebSocketMessageReceivedEventArgs^ EventArgs)
+	{
+		TSharedPtr<FHoloLensWebSocket> PinnedThis = WeakThis.Pin();
+		if (PinnedThis.IsValid())
+		{
+			try
+			{
+				bool bIsBinary = EventArgs->MessageType == SocketMessageType::Binary;
+				Windows::Storage::Streams::DataReader^ Reader = EventArgs->GetDataReader();
+				TArray<uint8> RawMessage;
+				RawMessage.AddUninitialized(Reader->UnconsumedBufferLength + 1);
+				Reader->ReadBytes(Platform::ArrayReference<uint8>(RawMessage.GetData(), RawMessage.Num() - 1));
+				if (!bIsBinary)
+				{
+					RawMessage.Last() = 0;
+				}
+
+				PinnedThis->GameThreadWork.Enqueue(
+					[WeakThis, bIsBinary, RawMessage = MoveTemp(RawMessage)]()
+				{
+					TSharedPtr<FHoloLensWebSocket> PinnedThis = WeakThis.Pin();
+					if (PinnedThis.IsValid())
+					{
+						if (!bIsBinary)
+						{
+							PinnedThis->OnMessage().Broadcast(UTF8_TO_TCHAR(RawMessage.GetData()));
+						}
+						PinnedThis->OnRawMessage().Broadcast(RawMessage.GetData(), RawMessage.Num() - 1, 0);
+					}
+				});
+			}
+			catch (Platform::Exception^ Ex)
+			{
+				UE_LOG(LogWebSockets, Warning, TEXT("Read error %s"), Ex->Message->Data());
+			}
+		}
+	});
+
+	Socket->Closed += ref new Windows::Foundation::TypedEventHandler<Windows::Networking::Sockets::IWebSocket^, WebSocketClosedEventArgs^>(
+		[WeakThis](Windows::Networking::Sockets::IWebSocket^, WebSocketClosedEventArgs^ EventArgs)
+	{
+		TSharedPtr<FHoloLensWebSocket> PinnedThis = WeakThis.Pin();
+		if (PinnedThis.IsValid())
+		{
+			int32 Code = EventArgs->Code;
+			FString Reason = EventArgs->Reason->Data();
+			PinnedThis->GameThreadWork.Enqueue(
+				[WeakThis, Code, Reason]()
+			{
+				TSharedPtr<FHoloLensWebSocket> PinnedThis = WeakThis.Pin();
+				if (PinnedThis.IsValid())
+				{
+					PinnedThis->FinishClose(Code, Reason);
+				}
+			});
+		}
+	});
+
+	try
+	{
+		ConnectAction = Socket->ConnectAsync(ref new Windows::Foundation::Uri(ref new Platform::String(*Uri)));
+	}
+	catch (Platform::Exception^ Ex)
+	{
+		FString ErrorMessage = Ex->Message->Data();
+		GameThreadWork.Enqueue(
+			[this, ErrorMessage]()
+		{
+			OnConnectionError().Broadcast(ErrorMessage);
+		});
+	}
+}
+
+void FHoloLensWebSocket::Close(int32 Code, const FString& Reason)
+{
+	bUserClose = true;
+	try
+	{
+		Socket->Close(static_cast<uint16>(Code), ref new Platform::String(*Reason));
+	}
+	catch (...)
+	{
+
+	}
+}
+
+bool FHoloLensWebSocket::IsConnected()
+{
+	return Writer != nullptr;
+}
+
+void FHoloLensWebSocket::Send(const FString& Data)
+{
+	if (Writer != nullptr)
+	{
+		try
+		{
+			Writer->WriteString(ref new Platform::String(*Data));
+			SendOperations.Add(Writer->StoreAsync());
+			OnMessageSent().Broadcast(Data);
+		}
+		catch (Platform::Exception^ Ex)
+		{
+			UE_LOG(LogWebSockets, Warning, TEXT("Send error %s"), Ex->Message->Data());
+		}
+	}
+}
+
+void FHoloLensWebSocket::Send(const void* Utf8Data, SIZE_T Size, bool bIsBinary)
+{
+	if (Writer != nullptr)
+	{
+		try
+		{
+			if (bIsBinary)
+			{
+				Socket->Control->MessageType = Windows::Networking::Sockets::SocketMessageType::Binary;
+			}
+			Platform::Array<uint8>^ WinRTArray = ref new Platform::Array<uint8>(static_cast<uint8*>(const_cast<void*>(Utf8Data)), Size);
+			Writer->WriteBytes(ref new Platform::Array<uint8>(static_cast<uint8*>(const_cast<void*>(Utf8Data)), Size));
+			SendOperations.Add(Writer->StoreAsync());
+			if (bIsBinary)
+			{
+				Socket->Control->MessageType = Windows::Networking::Sockets::SocketMessageType::Utf8;
+			}
+		}
+		catch (Platform::Exception^ Ex)
+		{
+			UE_LOG(LogWebSockets, Warning, TEXT("Send error %s"), Ex->Message->Data());
+		}
+	}
+}
+
+bool FHoloLensWebSocket::Tick(float DeltaTime)
+{
+	QUICK_SCOPE_CYCLE_COUNTER(STAT_FHoloLensWebSocket_Tick);
+
+	if (ConnectAction != nullptr)
+	{
+		switch (ConnectAction->Status)
+		{
+		case Windows::Foundation::AsyncStatus::Started:
+			break;
+
+		case Windows::Foundation::AsyncStatus::Canceled:
+			ConnectAction = nullptr;
+			break;
+
+		case Windows::Foundation::AsyncStatus::Error:
+			OnConnectionError().Broadcast(FString::FromInt(ConnectAction->ErrorCode.Value));
+			ConnectAction = nullptr;
+			break;
+
+		case Windows::Foundation::AsyncStatus::Completed:
+			Writer = ref new Windows::Storage::Streams::DataWriter(Socket->OutputStream);
+			Writer->UnicodeEncoding = Windows::Storage::Streams::UnicodeEncoding::Utf8;
+			OnConnected().Broadcast();
+			ConnectAction = nullptr;
+			break;
+
+		default:
+			break;
+		}
+	}
+
+	for (int32 i = 0; i < SendOperations.Num(); ++i)
+	{
+		if (SendOperations[i]->Status != Windows::Foundation::AsyncStatus::Started)
+		{
+			SendOperations.RemoveAtSwap(i);
+		}
+	}
+
+	TFunction<void()> WorkItem;
+	while (GameThreadWork.Dequeue(WorkItem))
+	{
+		WorkItem();
+	}
+
+	return true;
+}
+
+void FHoloLensWebSocket::FinishClose(int32 Code, const FString& Reason)
+{
+	try
+	{
+		Writer->DetachStream();
+	}
+	catch (...)
+	{
+
+	}
+	Writer = nullptr;
+
+	OnClosed().Broadcast(Code, Reason, bUserClose);
+}
+
+#endif
\ No newline at end of file
diff --git a/Engine/Platforms/Hololens/Source/Runtime/Online/WebSockets/Private/HoloLensWebSocket.h b/Engine/Platforms/Hololens/Source/Runtime/Online/WebSockets/Private/HoloLensWebSocket.h
new file mode 100644
--- /dev/null
+++ b/Engine/Platforms/Hololens/Source/Runtime/Online/WebSockets/Private/HoloLensWebSocket.h
@@ -0,0 +1,59 @@
+// Copyright Epic Games, Inc. All Rights Reserved.
+
+#pragma once
+
+#include "CoreMinimal.h"
+#include "Containers/Ticker.h"
+#include "Containers/Queue.h"
+
+#if WITH_WEBSOCKETS
+
+#include "IWebSocket.h"
+
+class FHoloLensWebSocket :
+	public IWebSocket,
+	public FTSTickerObjectBase,
+	public TSharedFromThis<FHoloLensWebSocket>
+{
+public:
+	FHoloLensWebSocket(const FString& Uri, const TArray<FString>& Protocols, const TMap<FString, FString>& Headers);
+
+	virtual void Connect() override;
+	virtual void Close(int32 Code = 1000, const FString& Reason = FString()) override;
+	virtual bool IsConnected() override;
+	virtual void Send(const FString& Data) override;
+	virtual void Send(const void* Utf8Data, SIZE_T Size, bool bIsBinary) override;
+
+	virtual FWebSocketConnectedEvent& OnConnected() override				{ return ConnectedEvent; }
+	virtual FWebSocketConnectionErrorEvent& OnConnectionError() override	{ return ConnectionErrorEvent; }
+	virtual FWebSocketClosedEvent& OnClosed() override						{ return ClosedEvent; }
+	virtual FWebSocketMessageEvent& OnMessage() override					{ return MessageEvent; }
+	virtual FWebSocketRawMessageEvent& OnRawMessage() override				{ return RawMessageEvent; }
+	virtual FWebSocketMessageSentEvent& OnMessageSent() override { return OnMessageSentEvent; }
+
+public:
+	virtual bool Tick(float DeltaTime) override;
+
+private:
+	void FinishClose(int32 Code, const FString& Reason);
+
+	TQueue<TFunction<void()>, EQueueMode::Mpsc> GameThreadWork;
+
+	Windows::Networking::Sockets::MessageWebSocket^ Socket;
+	Windows::Storage::Streams::DataWriter^ Writer;
+
+	Windows::Foundation::IAsyncAction^ ConnectAction;
+	TArray<Windows::Storage::Streams::DataWriterStoreOperation^> SendOperations;
+	FString Uri;
+
+	FWebSocketConnectedEvent ConnectedEvent;
+	FWebSocketConnectionErrorEvent ConnectionErrorEvent;
+	FWebSocketClosedEvent ClosedEvent;
+	FWebSocketMessageEvent MessageEvent;
+	FWebSocketRawMessageEvent RawMessageEvent;
+	FWebSocketMessageSentEvent OnMessageSentEvent;
+
+	bool bUserClose;
+};
+
+#endif
\ No newline at end of file
diff --git a/Engine/Platforms/Hololens/Source/Runtime/Online/WebSockets/Private/HoloLensWebSocketsManager.cpp b/Engine/Platforms/Hololens/Source/Runtime/Online/WebSockets/Private/HoloLensWebSocketsManager.cpp
new file mode 100644
--- /dev/null
+++ b/Engine/Platforms/Hololens/Source/Runtime/Online/WebSockets/Private/HoloLensWebSocketsManager.cpp
@@ -0,0 +1,12 @@
+// Copyright Epic Games, Inc. All Rights Reserved.
+
+#if WITH_WEBSOCKETS
+
+#include "HoloLensWebSocketsManager.h"
+
+TSharedRef<IWebSocket> FHoloLensWebSocketsManager::CreateWebSocket(const FString& Url, const TArray<FString>& Protocols, const TMap<FString, FString>& UpgradeHeaders)
+{
+	return MakeShared<FHoloLensWebSocket>(Url, Protocols, UpgradeHeaders);
+}
+
+#endif // #if WITH_WEBSOCKETS
\ No newline at end of file
diff --git a/Engine/Platforms/Hololens/Source/Runtime/Online/WebSockets/Private/HoloLensWebSocketsManager.h b/Engine/Platforms/Hololens/Source/Runtime/Online/WebSockets/Private/HoloLensWebSocketsManager.h
new file mode 100644
--- /dev/null
+++ b/Engine/Platforms/Hololens/Source/Runtime/Online/WebSockets/Private/HoloLensWebSocketsManager.h
@@ -0,0 +1,21 @@
+// Copyright Epic Games, Inc. All Rights Reserved.
+
+#pragma once
+
+#if WITH_WEBSOCKETS
+
+#include "IWebSocketsManager.h"
+#include "HoloLensWebSocket.h"
+
+class FHoloLensWebSocketsManager
+	: public IWebSocketsManager
+{
+private:
+	// IWebSocketsManager
+	virtual void InitWebSockets(TArrayView<const FString> Protocols) override {}
+	virtual void ShutdownWebSockets() override {}
+	virtual TSharedRef<IWebSocket> CreateWebSocket(const FString& Url, const TArray<FString>& Protocols, const TMap<FString, FString>& UpgradeHeaders) override;
+	// End IWebSocketsManager
+};
+
+#endif
diff --git a/Engine/Platforms/Hololens/Source/Runtime/Online/WebSockets/WebSockets_HoloLens.build.cs b/Engine/Platforms/Hololens/Source/Runtime/Online/WebSockets/WebSockets_HoloLens.build.cs
new file mode 100644
--- /dev/null
+++ b/Engine/Platforms/Hololens/Source/Runtime/Online/WebSockets/WebSockets_HoloLens.build.cs
@@ -0,0 +1,16 @@
+// Copyright Epic Games, Inc. All Rights Reserved.
+
+namespace UnrealBuildTool.Rules
+{
+	public class WebSockets_HoloLens : WebSockets
+	{
+		protected override bool bPlatformSupportsWinRTWebsockets
+		{
+			get => true;
+		}
+
+		public WebSockets_HoloLens(ReadOnlyTargetRules Target) : base(Target)
+		{
+		}
+	}
+}
diff --git a/Engine/Platforms/Hololens/Source/Runtime/OodleDataCompression/OodleDataCompression_HoloLens.build.cs b/Engine/Platforms/Hololens/Source/Runtime/OodleDataCompression/OodleDataCompression_HoloLens.build.cs
new file mode 100644
--- /dev/null
+++ b/Engine/Platforms/Hololens/Source/Runtime/OodleDataCompression/OodleDataCompression_HoloLens.build.cs
@@ -0,0 +1,51 @@
+// Copyright Epic Games, Inc. All Rights Reserved.
+
+namespace UnrealBuildTool.Rules
+{
+	public class OodleDataCompression_HoloLens : OodleDataCompression
+	{
+		protected override string LibRootDirectory { get { return GetModuleDirectoryForSubClass(GetType()).FullName; } }
+		protected override string ReleaseLibraryName
+		{
+			get
+			{
+				if (Target.WindowsPlatform.Architecture == WindowsArchitecture.x64) // emulation target, bBuildForEmulation
+				{
+					return "oo2core_winuwp64.lib";
+
+				}
+				else if (Target.WindowsPlatform.Architecture == WindowsArchitecture.ARM64) // device target, bBuildForDevice
+				{
+					return "oo2core_winuwparm64.lib";
+				}
+				else
+				{
+					throw new System.Exception("Unknown architecture for HoloLens platform!");
+				}
+			}
+		}
+		protected override string DebugLibraryName
+		{
+			get
+			{
+				if (Target.WindowsPlatform.Architecture == WindowsArchitecture.x64) // emulation target, bBuildForEmulation
+				{
+					return "oo2core_winuwp64_debug.lib";
+
+				}
+				else if (Target.WindowsPlatform.Architecture == WindowsArchitecture.ARM64) // device target, bBuildForDevice
+				{
+					return "oo2core_winuwparm64_debug.lib";
+				}
+				else
+				{
+					throw new System.Exception("Unknown architecture for HoloLens platform!");
+				}
+			}
+		}
+
+		public OodleDataCompression_HoloLens(ReadOnlyTargetRules Target) : base(Target)
+		{
+		}
+	}
+}
diff --git a/Engine/Platforms/Hololens/Source/Runtime/RHI/Private/HoloLensDynamicRHI.cpp b/Engine/Platforms/Hololens/Source/Runtime/RHI/Private/HoloLensDynamicRHI.cpp
new file mode 100644
--- /dev/null
+++ b/Engine/Platforms/Hololens/Source/Runtime/RHI/Private/HoloLensDynamicRHI.cpp
@@ -0,0 +1,119 @@
+// Copyright Epic Games, Inc. All Rights Reserved.
+
+#include "RHI.h"
+#include "Modules/ModuleManager.h"
+#include "Misc/ConfigCacheIni.h"
+#include "Misc/MessageDialog.h"
+
+static IDynamicRHIModule* LoadDynamicRHIModule()
+{
+	// command line overrides
+	bool bForceD3D11 = FParse::Param(FCommandLine::Get(), TEXT("d3d11")) || FParse::Param(FCommandLine::Get(), TEXT("dx11"));
+	bool bForceD3D12 = FParse::Param(FCommandLine::Get(), TEXT("d3d12")) || FParse::Param(FCommandLine::Get(), TEXT("dx12"));
+
+	if (!(bForceD3D11 || bForceD3D12))
+	{
+		//Default graphics RHI is only used if no command line option is specified
+		FConfigFile EngineSettings;
+		FString PlatformNameString = FPlatformProperties::PlatformName();
+		const TCHAR* PlatformName = *PlatformNameString;
+		FConfigCacheIni::LoadLocalIniFile(EngineSettings, TEXT("Engine"), true, PlatformName);
+		FString DefaultGraphicsRHI;
+		if (EngineSettings.GetString(TEXT("/Script/WindowsTargetPlatform.WindowsTargetSettings"), TEXT("DefaultGraphicsRHI"), DefaultGraphicsRHI))
+		{
+			FString NAME_DX11(TEXT("DefaultGraphicsRHI_DX11"));
+			FString NAME_DX12(TEXT("DefaultGraphicsRHI_DX12"));
+			if (DefaultGraphicsRHI == NAME_DX11)
+			{
+				bForceD3D11 = true;
+			}
+			else if (DefaultGraphicsRHI == NAME_DX12)
+			{
+				bForceD3D12 = true;
+			}
+		}
+	}
+	else if (bForceD3D11 && bForceD3D12)
+	{
+		UE_LOG(LogRHI, Fatal, TEXT("-d3d12 and -d3d11 are mutually exclusive options, but more than one was specified on the command-line."));
+	}
+
+	// Load the dynamic RHI module.
+	IDynamicRHIModule* DynamicRHIModule = NULL;
+
+	// Default to D3D12
+	if (!bForceD3D11)
+	{
+		FApp::SetGraphicsRHI(TEXT("DirectX 12"));
+		DynamicRHIModule = FModuleManager::LoadModulePtr<IDynamicRHIModule>(TEXT("D3D12RHI"));
+	}
+	else
+	{
+		FApp::SetGraphicsRHI(TEXT("DirectX 11"));
+		const TCHAR* D3D11RHIModuleName = TEXT("D3D11RHI");
+		DynamicRHIModule = &FModuleManager::LoadModuleChecked<IDynamicRHIModule>(D3D11RHIModuleName);
+	}
+
+	return DynamicRHIModule;
+}
+
+static IDynamicRHIModule* LoadWindowsMixedRealityDynamicRHIModule()
+{
+	IDynamicRHIModule* DynamicRHIModule = NULL;
+
+#if WITH_D3D12_RHI
+	bool bConfigRequestsD3D12 = false;
+	GConfig->GetBool(TEXT("/Script/HoloLensPlatformEditor.HoloLensTargetSettings"), TEXT("bUseD3D12RHI"), bConfigRequestsD3D12, GEngineIni);
+	const bool bForceD3D12 = FParse::Param(FCommandLine::Get(), TEXT("d3d12")) || FParse::Param(FCommandLine::Get(), TEXT("dx12"));
+
+	if (bForceD3D12 || bConfigRequestsD3D12)
+	{
+		DynamicRHIModule = &FModuleManager::LoadModuleChecked<IDynamicRHIModule>(TEXT("D3D12RHI"));
+	}
+	else
+#endif
+	{
+		// Load the dynamic RHI module.
+		auto& FMan = FModuleManager::Get();
+		DynamicRHIModule = (IDynamicRHIModule*)FMan.GetModule(TEXT("WindowsMixedRealityRHI"));
+		if (DynamicRHIModule == nullptr)
+		{
+			DynamicRHIModule = (IDynamicRHIModule*)FMan.LoadModule(TEXT("D3D11RHI"));
+		}
+	}
+
+	return DynamicRHIModule;
+}
+
+FDynamicRHI* PlatformCreateDynamicRHI()
+{
+	IDynamicRHIModule* DynamicRHIModule = NULL;
+
+	if (FModuleManager::Get().IsModuleLoaded(TEXT("WindowsMixedRealityRHI")))
+	{
+		// WindowsMixedReality uses a custom D3D11-based RHI and is incompatible with OpenXR.
+		DynamicRHIModule = LoadWindowsMixedRealityDynamicRHIModule();
+	}
+	else
+	{
+		DynamicRHIModule = LoadDynamicRHIModule();
+	}
+
+	// Create the dynamic RHI.
+	FDynamicRHI* DynamicRHI = NULL;
+	if (!DynamicRHIModule->IsSupported())
+	{
+		FMessageDialog::Open(EAppMsgType::Ok, NSLOCTEXT("HoloLens", "FailedToCreateHoloLens_RHI", "HoloLensRHI failure?"));
+		FPlatformMisc::RequestExit(1);
+		DynamicRHIModule = NULL;
+	}
+	else
+	{
+		DynamicRHI = DynamicRHIModule->CreateRHI(ERHIFeatureLevel::ES3_1);
+	}
+
+	GMaxRHIFeatureLevel = ERHIFeatureLevel::ES3_1;
+	GMaxRHIShaderPlatform = SP_D3D_ES3_1_HOLOLENS;
+
+	return DynamicRHI;
+}
diff --git a/Engine/Platforms/Hololens/Source/Runtime/RHI/RHI_HoloLens.build.cs b/Engine/Platforms/Hololens/Source/Runtime/RHI/RHI_HoloLens.build.cs
new file mode 100644
--- /dev/null
+++ b/Engine/Platforms/Hololens/Source/Runtime/RHI/RHI_HoloLens.build.cs
@@ -0,0 +1,19 @@
+// Copyright Epic Games, Inc. All Rights Reserved.
+
+namespace UnrealBuildTool.Rules
+{
+	public class RHI_HoloLens : RHI
+	{
+		public RHI_HoloLens(ReadOnlyTargetRules Target) : base(Target)
+		{
+			if (Target.bCompileAgainstEngine)
+			{
+				if (Target.Type != TargetRules.TargetType.Server)   // Dedicated servers should skip loading everything but NullDrv
+				{
+					DynamicallyLoadedModuleNames.Add("D3D11RHI");
+					DynamicallyLoadedModuleNames.Add("D3D12RHI");
+				}
+			}
+		}
+	}
+}
diff --git a/Engine/Platforms/Hololens/Source/Runtime/Slate/Private/Framework/Text/HoloLensPlatformTextField.cpp b/Engine/Platforms/Hololens/Source/Runtime/Slate/Private/Framework/Text/HoloLensPlatformTextField.cpp
new file mode 100644
--- /dev/null
+++ b/Engine/Platforms/Hololens/Source/Runtime/Slate/Private/Framework/Text/HoloLensPlatformTextField.cpp
@@ -0,0 +1,307 @@
+// Copyright Epic Games, Inc. All Rights Reserved.
+
+#include "Framework/Text/HoloLensPlatformTextField.h"
+
+using namespace Windows::UI::Text::Core;
+using namespace Windows::UI::Core;
+using namespace Windows::System;
+using namespace Windows::UI::ViewManagement;
+
+ref class VirtualKeyboardInputContext
+{
+internal:
+
+	VirtualKeyboardInputContext()
+	{
+	}
+
+	void CreateContext()
+	{
+		using namespace Windows::Foundation;
+
+		if (coreTextContext != nullptr)
+		{
+			return;
+		}
+
+		coreTextContext = CoreTextServicesManager::GetForCurrentView()->CreateEditContext();
+		coreTextContext->InputPaneDisplayPolicy = CoreTextInputPaneDisplayPolicy::Manual;
+
+		SelectionRequestedToken = coreTextContext->SelectionRequested += ref new TypedEventHandler<CoreTextEditContext^, CoreTextSelectionRequestedEventArgs^>(this, &VirtualKeyboardInputContext::OnSelectionRequested);
+		SelectionUpdatingToken = coreTextContext->SelectionUpdating += ref new TypedEventHandler<CoreTextEditContext^, CoreTextSelectionUpdatingEventArgs^>(this, &VirtualKeyboardInputContext::OnSelectionUpdating);
+		TextRequestedToken = coreTextContext->TextRequested += ref new TypedEventHandler<CoreTextEditContext^, CoreTextTextRequestedEventArgs^>(this, &VirtualKeyboardInputContext::OnTextRequested);
+		TextUpdatingToken = coreTextContext->TextUpdating += ref new TypedEventHandler<CoreTextEditContext^, CoreTextTextUpdatingEventArgs^>(this, &VirtualKeyboardInputContext::OnTextUpdating);
+		FormatUpdatingToken = coreTextContext->FormatUpdating += ref new TypedEventHandler<CoreTextEditContext^, CoreTextFormatUpdatingEventArgs^>(this, &VirtualKeyboardInputContext::OnFormatUpdating);
+
+		KeyDownCoreWindowToken = CoreWindow::GetForCurrentThread()->KeyDown += ref new TypedEventHandler<CoreWindow^, KeyEventArgs^>(this, &VirtualKeyboardInputContext::OnKeyDownCoreWindow);
+
+		TextEntryWidget->OnSelectionChanged.BindLambda([this]() { OnSelectionChanged(); });
+	}
+
+	void DeleteContext()
+	{
+		if (TextEntryWidget)
+		{
+			TextEntryWidget->OnSelectionChanged.Unbind();
+		}
+
+		CoreWindow::GetForCurrentThread()->KeyDown -= KeyDownCoreWindowToken;
+
+		if (coreTextContext)
+		{
+			coreTextContext->SelectionRequested -= SelectionRequestedToken;
+			coreTextContext->SelectionUpdating -= SelectionUpdatingToken;
+			coreTextContext->TextRequested -= TextRequestedToken;
+			coreTextContext->TextUpdating -= TextUpdatingToken;
+			coreTextContext->FormatUpdating -= FormatUpdatingToken;
+		}
+
+		coreTextContext = nullptr;
+	}
+
+	void OnSelectionRequested(CoreTextEditContext^, CoreTextSelectionRequestedEventArgs^ args)
+	{
+		int start, end;
+		if (TextEntryWidget->GetSelection(start, end))
+		{
+			CoreTextRange range;
+			range.StartCaretPosition = start;
+			range.EndCaretPosition = end;
+			args->Request->Selection = range;
+		}
+	}
+
+	void OnSelectionUpdating(CoreTextEditContext^, CoreTextSelectionUpdatingEventArgs^ args)
+	{
+		TextEntryWidget->SetSelectionFromVirtualKeyboard(args->Selection.StartCaretPosition, args->Selection.EndCaretPosition);
+	}
+
+	void OnTextRequested(CoreTextEditContext^, CoreTextTextRequestedEventArgs^ args)
+	{
+		args->Request->Text = ref new Platform::String(*TextEntryWidget->GetText().ToString());
+	}
+
+	void OnTextUpdating(CoreTextEditContext^, CoreTextTextUpdatingEventArgs^ args)
+	{
+		FString text = TextEntryWidget->GetText().ToString();
+		uint32 start = args->Range.StartCaretPosition;
+		uint32 end = args->Range.EndCaretPosition;
+		FString newText = args->Text->Data();
+
+		if (start != end)
+		{
+			text.RemoveAt(start, end - start);
+		}
+		text.InsertAt(start, newText);
+
+		TextEntryWidget->SetTextFromVirtualKeyboard(FText::FromString(text), ETextEntryType::TextEntryUpdated);
+		TextEntryWidget->SetSelectionFromVirtualKeyboard(args->NewSelection.StartCaretPosition, args->NewSelection.EndCaretPosition);
+	}
+
+	// Needed for autocomplete
+	void OnFormatUpdating(CoreTextEditContext^, CoreTextFormatUpdatingEventArgs^ args)
+	{
+	}
+
+	void OnSelectionChanged()
+	{
+		int start, end;
+		if (TextEntryWidget->GetSelection(start, end))
+		{
+			CoreTextRange selRange;
+			selRange.StartCaretPosition = start;
+			selRange.EndCaretPosition = end;
+			coreTextContext->NotifySelectionChanged(selRange);
+		}
+	}
+
+	void Start(TSharedPtr<IVirtualKeyboardEntry> textEntryWidget)
+	{
+		TextEntryWidget = textEntryWidget;
+		CreateContext();
+		coreTextContext->NotifyFocusEnter();
+		InputPane::GetForCurrentView()->TryShow();
+	}
+
+	void Stop()
+	{
+		if (InputPane::GetForCurrentView() != nullptr)
+		{
+			InputPane::GetForCurrentView()->TryHide();
+		}
+		
+		if (coreTextContext)
+		{
+			coreTextContext->NotifyFocusLeave();
+		}
+
+		DeleteContext();
+		TextEntryWidget = nullptr;
+	}
+
+	void OnBack()
+	{
+		FString text = TextEntryWidget->GetText().ToString();
+		int oldLen = text.Len();
+		if (oldLen == 0)
+		{
+			return;
+		}
+
+		int start, end;
+		if (!TextEntryWidget->GetSelection(start, end))
+		{
+			start = end = text.Len();
+		}
+
+		if (start == end)
+		{
+			//no selection
+			if (start == 0)
+			{
+				return;
+			}
+
+			start--;
+		}
+
+		text.RemoveAt(start, end - start);
+
+		CoreTextRange txtRange, selRange;
+		txtRange.StartCaretPosition = start;
+		txtRange.EndCaretPosition = end;
+		selRange.StartCaretPosition = start;
+		selRange.EndCaretPosition = start;
+
+		TextEntryWidget->SetTextFromVirtualKeyboard(FText::FromString(text), ETextEntryType::TextEntryUpdated);
+		TextEntryWidget->SetSelectionFromVirtualKeyboard(start, start);
+		coreTextContext->NotifyTextChanged(txtRange, oldLen - (end - start), selRange);
+	}
+
+	void OnLeft()
+	{
+		FString text = TextEntryWidget->GetText().ToString();
+		int oldLen = text.Len();
+		if (oldLen == 0)
+		{
+			return;
+		}
+
+		int start, end;
+		if (!TextEntryWidget->GetSelection(start, end))
+		{
+			start = end = text.Len();
+		}
+
+		if (start == end)
+		{
+			if (start == 0)
+			{
+				return;
+			}
+
+			--start;
+		}
+		end = start;
+
+		CoreTextRange selRange;
+		selRange.StartCaretPosition = start;
+		selRange.EndCaretPosition = end;
+		coreTextContext->NotifySelectionChanged(selRange);
+		TextEntryWidget->SetSelectionFromVirtualKeyboard(start, end);
+	}
+
+	void OnRight()
+	{
+		FString text = TextEntryWidget->GetText().ToString();
+		int oldLen = text.Len();
+		if (oldLen == 0)
+		{
+			return;
+		}
+
+		int start, end;
+		if (!TextEntryWidget->GetSelection(start, end))
+		{
+			start = end = text.Len();
+		}
+
+		if (start == end)
+		{
+			if (end == oldLen)
+			{
+				return;
+			}
+
+			++end;
+		}
+		start = end;
+
+		CoreTextRange selRange;
+		selRange.StartCaretPosition = start;
+		selRange.EndCaretPosition = end;
+		coreTextContext->NotifySelectionChanged(selRange);
+		TextEntryWidget->SetSelectionFromVirtualKeyboard(start, end);
+
+	}
+
+	void OnKeyDownCoreWindow(CoreWindow^ sender, KeyEventArgs^ args)
+	{
+		switch (args->VirtualKey)
+		{
+		case VirtualKey::Back:
+			OnBack();
+			break;
+		case VirtualKey::Left:
+			OnLeft();
+			break;
+		case VirtualKey::Right:
+			OnRight();
+			break;
+		}
+	}
+
+
+private:
+	Windows::Foundation::EventRegistrationToken SelectionRequestedToken;
+	Windows::Foundation::EventRegistrationToken SelectionUpdatingToken;
+	Windows::Foundation::EventRegistrationToken TextRequestedToken;
+	Windows::Foundation::EventRegistrationToken TextUpdatingToken;
+	Windows::Foundation::EventRegistrationToken FormatUpdatingToken;
+	Windows::Foundation::EventRegistrationToken KeyDownCoreWindowToken;
+
+	Windows::UI::Text::Core::CoreTextEditContext^ coreTextContext = nullptr;
+	TSharedPtr<IVirtualKeyboardEntry> TextEntryWidget;
+};
+
+FHoloLensPlatformTextField::FHoloLensPlatformTextField()
+{
+	PreLoadMapHandle = FCoreUObjectDelegates::PreLoadMap.AddRaw(this, &FHoloLensPlatformTextField::LevelChanging);
+
+	inputContext = ref new VirtualKeyboardInputContext();
+}
+
+FHoloLensPlatformTextField::~FHoloLensPlatformTextField()
+{
+	FCoreUObjectDelegates::PreLoadMap.Remove(PreLoadMapHandle);
+}
+
+void FHoloLensPlatformTextField::LevelChanging(const FString& MapName)
+{
+	if (inputContext)
+	{
+		inputContext->Stop();
+	}
+}
+
+void FHoloLensPlatformTextField::ShowVirtualKeyboard(bool bShow, int32 UserIndex, TSharedPtr<IVirtualKeyboardEntry> TextEntryWidget) 
+{
+	if (bShow)
+	{
+		inputContext->Start(TextEntryWidget);
+	}
+	else
+	{
+		inputContext->Stop();
+	}
+}
diff --git a/Engine/Platforms/Hololens/Source/Runtime/Slate/Public/Framework/Text/HololensPlatformTextField.h b/Engine/Platforms/Hololens/Source/Runtime/Slate/Public/Framework/Text/HololensPlatformTextField.h
new file mode 100644
--- /dev/null
+++ b/Engine/Platforms/Hololens/Source/Runtime/Slate/Public/Framework/Text/HololensPlatformTextField.h
@@ -0,0 +1,26 @@
+// Copyright Epic Games, Inc. All Rights Reserved.
+
+#pragma once
+
+#include "CoreMinimal.h"
+#include "Framework/Application/IPlatformTextField.h"
+
+ref class VirtualKeyboardInputContext;
+
+class SLATE_API FHoloLensPlatformTextField : public IPlatformTextField
+{
+public:
+	FHoloLensPlatformTextField();
+	~FHoloLensPlatformTextField();
+
+	void LevelChanging(const FString& MapName);
+	
+	virtual void ShowVirtualKeyboard(bool bShow, int32 UserIndex, TSharedPtr<IVirtualKeyboardEntry> TextEntryWidget) override;
+	virtual bool AllowMoveCursor() override { return true; };
+	
+private:
+	VirtualKeyboardInputContext^ inputContext = nullptr;
+	FDelegateHandle PreLoadMapHandle;
+};
+
+typedef FHoloLensPlatformTextField FPlatformTextField;
diff --git a/Engine/Platforms/Hololens/Source/Runtime/Slate/Slate_HoloLens.build.cs b/Engine/Platforms/Hololens/Source/Runtime/Slate/Slate_HoloLens.build.cs
new file mode 100644
--- /dev/null
+++ b/Engine/Platforms/Hololens/Source/Runtime/Slate/Slate_HoloLens.build.cs
@@ -0,0 +1,11 @@
+// Copyright Epic Games, Inc. All Rights Reserved.
+
+namespace UnrealBuildTool.Rules
+{
+	public class Slate_HoloLens : Slate
+	{
+		public Slate_HoloLens(ReadOnlyTargetRules Target) : base(Target)
+		{
+		}
+	}
+}
diff --git a/Engine/Platforms/Hololens/Source/Runtime/Sockets/Private/SocketSubsystemHoloLens.cpp b/Engine/Platforms/Hololens/Source/Runtime/Sockets/Private/SocketSubsystemHoloLens.cpp
new file mode 100644
--- /dev/null
+++ b/Engine/Platforms/Hololens/Source/Runtime/Sockets/Private/SocketSubsystemHoloLens.cpp
@@ -0,0 +1,215 @@
+// Copyright Epic Games, Inc. All Rights Reserved.
+
+#include "SocketSubsystemHoloLens.h"
+#include "SocketSubsystemModule.h"
+#include "Modules/ModuleManager.h"
+
+#include "BSDSockets/IPAddressBSD.h"
+
+#include "Microsoft/AllowMicrosoftPlatformTypes.h"
+THIRD_PARTY_INCLUDES_START
+#include "Iphlpapi.h"
+THIRD_PARTY_INCLUDES_END
+#include "Microsoft/HideMicrosoftPlatformTypes.h"
+
+
+FSocketSubsystemHoloLens* FSocketSubsystemHoloLens::SocketSingleton = nullptr;
+
+
+FName CreateSocketSubsystem(FSocketSubsystemModule& SocketSubsystemModule)
+{
+	FName SubsystemName(TEXT("HoloLens"));
+
+	// Create and register our singleton factory with the main online subsystem for easy access
+	FSocketSubsystemHoloLens* SocketSubsystem = FSocketSubsystemHoloLens::Create();
+	FString Error;
+
+	if (SocketSubsystem->Init(Error))
+	{
+		SocketSubsystemModule.RegisterSocketSubsystem(SubsystemName, SocketSubsystem);
+
+		return SubsystemName;
+	}
+
+	FSocketSubsystemHoloLens::Destroy();
+
+	return NAME_None;
+}
+
+
+void DestroySocketSubsystem(FSocketSubsystemModule& SocketSubsystemModule)
+{
+	SocketSubsystemModule.UnregisterSocketSubsystem(FName(TEXT("HoloLens")));
+	FSocketSubsystemHoloLens::Destroy();
+}
+
+
+/* FSocketSubsystemHoloLens interface
+*****************************************************************************/
+
+FSocketSubsystemHoloLens* FSocketSubsystemHoloLens::Create()
+{
+	if (SocketSingleton == nullptr)
+	{
+		SocketSingleton = new FSocketSubsystemHoloLens();
+	}
+
+	return SocketSingleton;
+}
+
+void FSocketSubsystemHoloLens::Destroy()
+{
+	if (SocketSingleton != nullptr)
+	{
+		SocketSingleton->Shutdown();
+		delete SocketSingleton;
+		SocketSingleton = nullptr;
+	}
+}
+
+
+/* FSocketSubsystemBSD overrides
+*****************************************************************************/
+
+FSocket* FSocketSubsystemHoloLens::CreateSocket(const FName& SocketType, const FString& SocketDescription, const FName& ProtocolType)
+{
+	FSocketBSD* NewSocket = static_cast<FSocketBSD*>(FSocketSubsystemBSD::CreateSocket(SocketType, SocketDescription, ProtocolType));
+
+	if (NewSocket != nullptr)
+	{
+		NewSocket->SetIPv6Only(false);
+	}
+	else	
+	{
+		UE_LOG(LogSockets, Warning, TEXT("Failed to create socket %s [%s]"), *SocketType.ToString(), *SocketDescription);
+	}
+
+	return NewSocket;
+}
+
+
+bool FSocketSubsystemHoloLens::Init(FString& Error)
+{
+	bool bSuccess = false;
+
+	if (bTriedToInit == false)
+	{
+		bTriedToInit = true;
+		WSADATA WSAData;
+
+		// initialize WSA
+		int32 Code = WSAStartup(0x0101, &WSAData);
+
+		if (Code == 0)
+		{
+			bSuccess = true;
+			UE_LOG(LogInit, Log, TEXT("%s: version %i.%i (%i.%i), MaxSocks=%i, MaxUdp=%i"),
+				GetSocketAPIName(),
+				WSAData.wVersion >> 8, WSAData.wVersion & 0xFF,
+				WSAData.wHighVersion >> 8, WSAData.wHighVersion & 0xFF,
+				WSAData.iMaxSockets, WSAData.iMaxUdpDg);
+		}
+		else
+		{
+			Error = FString::Printf(TEXT("WSAStartup failed (%s)"), GetSocketError(TranslateErrorCode(Code)));
+		}
+	}
+
+	return bSuccess;
+}
+
+
+void FSocketSubsystemHoloLens::Shutdown(void)
+{
+	WSACleanup();
+}
+
+
+ESocketErrors FSocketSubsystemHoloLens::GetLastErrorCode()
+{
+	return TranslateErrorCode(WSAGetLastError());
+}
+
+ESocketErrors FSocketSubsystemHoloLens::TranslateErrorCode(int32 Code)
+{
+	// handle the generic -1 error
+	if (Code == SOCKET_ERROR)
+	{
+		return GetLastErrorCode();
+	}
+
+	switch (Code)
+	{
+	case 0: return SE_NO_ERROR;
+	case ERROR_INVALID_HANDLE: return SE_ECONNRESET; // invalid socket handle catch
+	case WSAEINTR: return SE_EINTR;
+	case WSAEBADF: return SE_EBADF;
+	case WSAEACCES: return SE_EACCES;
+	case WSAEFAULT: return SE_EFAULT;
+	case WSAEINVAL: return SE_EINVAL;
+	case WSAEMFILE: return SE_EMFILE;
+	case WSAEWOULDBLOCK: return SE_EWOULDBLOCK;
+	case WSAEINPROGRESS: return SE_EINPROGRESS;
+	case WSAEALREADY: return SE_EALREADY;
+	case WSAENOTSOCK: return SE_ENOTSOCK;
+	case WSAEDESTADDRREQ: return SE_EDESTADDRREQ;
+	case WSAEMSGSIZE: return SE_EMSGSIZE;
+	case WSAEPROTOTYPE: return SE_EPROTOTYPE;
+	case WSAENOPROTOOPT: return SE_ENOPROTOOPT;
+	case WSAEPROTONOSUPPORT: return SE_EPROTONOSUPPORT;
+	case WSAESOCKTNOSUPPORT: return SE_ESOCKTNOSUPPORT;
+	case WSAEOPNOTSUPP: return SE_EOPNOTSUPP;
+	case WSAEPFNOSUPPORT: return SE_EPFNOSUPPORT;
+	case WSAEAFNOSUPPORT: return SE_EAFNOSUPPORT;
+	case WSAEADDRINUSE: return SE_EADDRINUSE;
+	case WSAEADDRNOTAVAIL: return SE_EADDRNOTAVAIL;
+	case WSAENETDOWN: return SE_ENETDOWN;
+	case WSAENETUNREACH: return SE_ENETUNREACH;
+	case WSAENETRESET: return SE_ENETRESET;
+	case WSAECONNABORTED: return SE_ECONNABORTED;
+	case WSAECONNRESET: return SE_ECONNRESET;
+	case WSAENOBUFS: return SE_ENOBUFS;
+	case WSAEISCONN: return SE_EISCONN;
+	case WSAENOTCONN: return SE_ENOTCONN;
+	case WSAESHUTDOWN: return SE_ESHUTDOWN;
+	case WSAETOOMANYREFS: return SE_ETOOMANYREFS;
+	case WSAETIMEDOUT: return SE_ETIMEDOUT;
+	case WSAECONNREFUSED: return SE_ECONNREFUSED;
+	case WSAELOOP: return SE_ELOOP;
+	case WSAENAMETOOLONG: return SE_ENAMETOOLONG;
+	case WSAEHOSTDOWN: return SE_EHOSTDOWN;
+	case WSAEHOSTUNREACH: return SE_EHOSTUNREACH;
+	case WSAENOTEMPTY: return SE_ENOTEMPTY;
+	case WSAEPROCLIM: return SE_EPROCLIM;
+	case WSAEUSERS: return SE_EUSERS;
+	case WSAEDQUOT: return SE_EDQUOT;
+	case WSAESTALE: return SE_ESTALE;
+	case WSAEREMOTE: return SE_EREMOTE;
+	case WSAEDISCON: return SE_EDISCON;
+	case WSASYSNOTREADY: return SE_SYSNOTREADY;
+	case WSAVERNOTSUPPORTED: return SE_VERNOTSUPPORTED;
+	case WSANOTINITIALISED: return SE_NOTINITIALISED;
+	case WSAHOST_NOT_FOUND: return SE_HOST_NOT_FOUND;
+	case WSATRY_AGAIN: return SE_TRY_AGAIN;
+	case WSANO_RECOVERY: return SE_NO_RECOVERY;
+	case WSANO_DATA: return SE_NO_DATA;
+		// case : return SE_UDP_ERR_PORT_UNREACH; //@TODO Find it's replacement
+	}
+
+	UE_LOG(LogSockets, Warning, TEXT("Unhandled socket error! Error Code: %d"), Code);
+	check(0);
+
+	return SE_NO_ERROR;
+}
+
+
+bool FSocketSubsystemHoloLens::HasNetworkDevice()
+{
+	return true;
+}
+
+const TCHAR* FSocketSubsystemHoloLens::GetSocketAPIName() const
+{
+	return TEXT("WinSock");
+}
+
diff --git a/Engine/Platforms/Hololens/Source/Runtime/Sockets/Private/SocketSubsystemHoloLens.h b/Engine/Platforms/Hololens/Source/Runtime/Sockets/Private/SocketSubsystemHoloLens.h
new file mode 100644
--- /dev/null
+++ b/Engine/Platforms/Hololens/Source/Runtime/Sockets/Private/SocketSubsystemHoloLens.h
@@ -0,0 +1,62 @@
+// Copyright Epic Games, Inc. All Rights Reserved.
+
+#pragma once
+
+#include "SocketSubsystem.h"
+#include "BSDSockets/SocketSubsystemBSD.h"
+#include "BSDSockets/SocketsBSD.h"
+#include "SocketSubsystemPackage.h"
+
+/**
+ * HoloLens specific socket subsystem implementation
+ */
+class FSocketSubsystemHoloLens
+	: public FSocketSubsystemBSD
+{
+public:
+
+	/** Default Constructor. */
+	FSocketSubsystemHoloLens() :
+		bTriedToInit(false)
+	{ }
+
+	/** Virtual destructor. */
+	virtual ~FSocketSubsystemHoloLens() { }
+
+public:
+
+	// FSocketSubsystemBSD overrides
+
+	virtual class FSocket* CreateSocket(const FName& SocketType, const FString& SocketDescription, const FName& ProtocolType) override;
+	virtual bool HasNetworkDevice() override;
+	virtual ESocketErrors GetLastErrorCode() override;
+	virtual const TCHAR* GetSocketAPIName() const override;
+	virtual bool Init(FString& Error) override;
+	virtual void Shutdown() override;
+	virtual ESocketErrors TranslateErrorCode(int32 Code) override;
+
+	virtual FName GetDefaultSocketProtocolFamily() const override
+	{
+		return FNetworkProtocolTypes::IPv6;
+	}
+
+PACKAGE_SCOPE:
+
+	/**
+	* Singleton interface for this subsystem
+	*
+	* @return the only instance of this subsystem
+	*/
+	static FSocketSubsystemHoloLens* Create();
+
+	/** Performs HoloLens specific socket clean up. */
+	static void Destroy();
+
+protected:
+
+	/** Holds a flag indicating whether Init() has been called before or not. */
+	bool bTriedToInit;
+
+	/** Holds the single instantiation of this subsystem. */
+	static FSocketSubsystemHoloLens* SocketSingleton;
+};
\ No newline at end of file
diff --git a/Engine/Platforms/Hololens/Source/Runtime/Sockets/Sockets_HoloLens.build.cs b/Engine/Platforms/Hololens/Source/Runtime/Sockets/Sockets_HoloLens.build.cs
new file mode 100644
--- /dev/null
+++ b/Engine/Platforms/Hololens/Source/Runtime/Sockets/Sockets_HoloLens.build.cs
@@ -0,0 +1,11 @@
+// Copyright Epic Games, Inc. All Rights Reserved.
+
+namespace UnrealBuildTool.Rules
+{
+	public class Sockets_HoloLens : Sockets
+	{
+		public Sockets_HoloLens(ReadOnlyTargetRules Target) : base(Target)
+		{
+		}
+	}
+}
diff --git a/Engine/Platforms/Hololens/Source/Runtime/TraceLog/Private/HoloLensTrace.cpp b/Engine/Platforms/Hololens/Source/Runtime/TraceLog/Private/HoloLensTrace.cpp
new file mode 100644
--- /dev/null
+++ b/Engine/Platforms/Hololens/Source/Runtime/TraceLog/Private/HoloLensTrace.cpp
@@ -0,0 +1,274 @@
+// Copyright Epic Games, Inc. All Rights Reserved.
+
+#include "Trace/Config.h"
+
+#if UE_TRACE_ENABLED && PLATFORM_HOLOLENS
+
+#include "Windows/AllowWindowsPlatformTypes.h"
+#	include "Windows/WindowsHWrapper.h"
+#	define _WINSOCK_DEPRECATED_NO_WARNINGS
+#	include <winsock2.h>
+#	include <ws2tcpip.h>
+#	pragma comment(lib, "ws2_32.lib")
+#include "Windows/HideWindowsPlatformTypes.h"
+
+#pragma warning(push)
+#pragma warning(disable : 6031) // WSAStartup() return ignore  - we're error tolerant
+
+namespace UE {
+namespace Trace {
+namespace Private {
+
+////////////////////////////////////////////////////////////////////////////////
+UPTRINT ThreadCreate(const ANSICHAR* Name, void (*Entry)())
+{
+	DWORD (WINAPI *WinApiThunk)(void*) = [] (void* Param) -> DWORD
+	{
+		typedef void (*EntryType)(void);
+		(EntryType(Param))();
+		return 0;
+	};
+
+	HANDLE Handle = CreateThread(nullptr, 0, WinApiThunk, (void*)Entry, 0, nullptr);
+	return UPTRINT(Handle);
+}
+
+////////////////////////////////////////////////////////////////////////////////
+void ThreadSleep(uint32 Milliseconds)
+{
+	Sleep(Milliseconds);
+}
+
+////////////////////////////////////////////////////////////////////////////////
+void ThreadJoin(UPTRINT Handle)
+{
+	WaitForSingleObject(HANDLE(Handle), INFINITE);
+}
+
+////////////////////////////////////////////////////////////////////////////////
+void ThreadDestroy(UPTRINT Handle)
+{
+	CloseHandle(HANDLE(Handle));
+}
+
+
+
+////////////////////////////////////////////////////////////////////////////////
+uint64 TimeGetFrequency()
+{
+	LARGE_INTEGER Value;
+	QueryPerformanceFrequency(&Value);
+	return Value.QuadPart;
+}
+
+////////////////////////////////////////////////////////////////////////////////
+TRACELOG_API uint64 TimeGetTimestamp()
+{
+	LARGE_INTEGER Value;
+	QueryPerformanceCounter(&Value);
+	return Value.QuadPart;
+}
+
+
+
+////////////////////////////////////////////////////////////////////////////////
+static void TcpSocketInitialize()
+{
+	WSADATA WsaData;
+	WSAStartup(MAKEWORD(2, 2), &WsaData);
+}
+
+////////////////////////////////////////////////////////////////////////////////
+static bool TcpSocketSetNonBlocking(SOCKET Socket, bool bNonBlocking)
+{
+	unsigned long NonBlockingMode = !!bNonBlocking;
+	return ioctlsocket(Socket, FIONBIO, &NonBlockingMode) != SOCKET_ERROR;
+}
+
+////////////////////////////////////////////////////////////////////////////////
+UPTRINT TcpSocketConnect(const ANSICHAR* Host, uint16 Port)
+{
+	TcpSocketInitialize();
+
+	struct FAddrInfoPtr
+	{
+					~FAddrInfoPtr()	{ freeaddrinfo(Value); }
+		addrinfo*	operator -> ()	{ return Value; }
+		addrinfo**	operator & ()	{ return &Value; }
+		addrinfo*	Value;
+	};
+
+	FAddrInfoPtr Info;
+	addrinfo Hints = {};
+	Hints.ai_family = AF_INET;
+	Hints.ai_socktype = SOCK_STREAM;
+	Hints.ai_protocol = IPPROTO_TCP;
+	if (getaddrinfo(Host, nullptr, &Hints, &Info))
+	{
+		return 0;
+	}
+
+	if (&Info == nullptr)
+	{
+		return 0;
+	}
+
+	auto* SockAddr = (sockaddr_in*)Info->ai_addr;
+	SockAddr->sin_port = htons(Port);
+
+	// socket() will create a socket with overlapped IO support which we don't
+	// want as it complicates sharing Io*() API with FileOpen(). So we use
+	// WSASocket instead which affords us more control over socket properties
+	SOCKET Socket = WSASocket(AF_INET, SOCK_STREAM, IPPROTO_TCP, nullptr, 0, WSA_FLAG_NO_HANDLE_INHERIT);
+	if (Socket == INVALID_SOCKET)
+	{
+		return 0;
+	}
+
+	int Result = connect(Socket, Info->ai_addr, int(Info->ai_addrlen));
+	if (Result == SOCKET_ERROR)
+	{
+		closesocket(Socket);
+		return 0;
+	}
+
+	if (!TcpSocketSetNonBlocking(Socket, 0))
+	{
+		closesocket(Socket);
+		return 0;
+	}
+
+	return UPTRINT(Socket) + 1;
+}
+
+////////////////////////////////////////////////////////////////////////////////
+UPTRINT TcpSocketListen(uint16 Port)
+{
+	TcpSocketInitialize();
+
+	// See TcpSocketConnect() for why WSASocket() is used here.
+	SOCKET Socket = WSASocket(AF_INET, SOCK_STREAM, IPPROTO_TCP, nullptr, 0, WSA_FLAG_NO_HANDLE_INHERIT);
+	if (Socket == INVALID_SOCKET)
+	{
+		return 0;
+	}
+
+	sockaddr_in SockAddr;
+	SockAddr.sin_family = AF_INET;
+	SockAddr.sin_addr.s_addr = 0;
+	SockAddr.sin_port = htons(Port);
+	int Result = bind(Socket, (SOCKADDR*)&SockAddr, sizeof(SockAddr));
+	if (Result == INVALID_SOCKET)
+	{
+		closesocket(Socket);
+		return 0;
+	}
+
+	Result = listen(Socket, 1);
+	if (Result == INVALID_SOCKET)
+	{
+		closesocket(Socket);
+		return 0;
+	}
+
+	if (!TcpSocketSetNonBlocking(Socket, 1))
+	{
+		closesocket(Socket);
+		return 0;
+	}
+
+	return UPTRINT(Socket) + 1;
+}
+
+////////////////////////////////////////////////////////////////////////////////
+int32 TcpSocketAccept(UPTRINT Socket, UPTRINT& Out)
+{
+	SOCKET Inner = Socket - 1;
+
+	Inner = accept(Inner, nullptr, nullptr);
+	if (Inner == INVALID_SOCKET)
+	{
+		return (WSAGetLastError() == WSAEWOULDBLOCK) - 1; // 0 if would block else -1
+	}
+
+	if (!TcpSocketSetNonBlocking(Inner, 0))
+	{
+		closesocket(Inner);
+		return 0;
+	}
+
+	Out = UPTRINT(Inner) + 1;
+	return 1;
+}
+
+////////////////////////////////////////////////////////////////////////////////
+bool TcpSocketHasData(UPTRINT Socket)
+{
+	SOCKET Inner = Socket - 1;
+	fd_set FdSet = { 1, { Inner }, };
+	TIMEVAL TimeVal = {};
+	return (select(0, &FdSet, nullptr, nullptr, &TimeVal) != 0);
+}
+
+
+
+////////////////////////////////////////////////////////////////////////////////
+bool IoWrite(UPTRINT Handle, const void* Data, uint32 Size)
+{
+	HANDLE Inner = HANDLE(Handle - 1);
+
+	DWORD BytesWritten = 0;
+	if (!WriteFile(Inner, (const char*)Data, Size, &BytesWritten, nullptr))
+	{
+		return false;
+	}
+
+	return (BytesWritten == Size);
+}
+
+////////////////////////////////////////////////////////////////////////////////
+int32 IoRead(UPTRINT Handle, void* Data, uint32 Size)
+{
+	HANDLE Inner = HANDLE(Handle - 1);
+
+	DWORD BytesRead = 0;
+	if (!ReadFile(Inner, (char*)Data, Size, &BytesRead, nullptr))
+	{
+		return -1;
+	}
+
+	return BytesRead;
+}
+
+////////////////////////////////////////////////////////////////////////////////
+void IoClose(UPTRINT Handle)
+{
+	HANDLE Inner = HANDLE(Handle - 1);
+	CloseHandle(Inner);
+}
+
+
+
+////////////////////////////////////////////////////////////////////////////////
+UPTRINT FileOpen(const ANSICHAR* Path)
+{
+	DWORD Access = GENERIC_WRITE;
+	DWORD Share = FILE_SHARE_READ;
+	DWORD Disposition = CREATE_ALWAYS;
+	DWORD Flags = FILE_ATTRIBUTE_NORMAL;
+	HANDLE Out = CreateFile2((LPCWSTR)Path, Access, Share, Disposition, nullptr);
+	if (Out == INVALID_HANDLE_VALUE)
+	{
+		return 0;
+	}
+
+	return UPTRINT(Out) + 1;
+}
+
+} // namespace Private
+} // namespace Trace
+} // namespace UE
+
+#pragma warning(pop)
+
+#endif // UE_TRACE_ENABLED
diff --git a/Engine/Platforms/Hololens/Source/Runtime/TraceLog/TraceLog_HoloLens.build.cs b/Engine/Platforms/Hololens/Source/Runtime/TraceLog/TraceLog_HoloLens.build.cs
new file mode 100644
--- /dev/null
+++ b/Engine/Platforms/Hololens/Source/Runtime/TraceLog/TraceLog_HoloLens.build.cs
@@ -0,0 +1,11 @@
+// Copyright Epic Games, Inc. All Rights Reserved.
+
+namespace UnrealBuildTool.Rules
+{
+	public class TraceLog_HoloLens : TraceLog
+	{
+		public TraceLog_HoloLens(ReadOnlyTargetRules Target) : base(Target)
+		{
+		}
+	}
+}
diff --git a/Engine/Platforms/Hololens/Source/Runtime/Windows/D3D11RHI/D3D11RHI_HoloLens.build.cs b/Engine/Platforms/Hololens/Source/Runtime/Windows/D3D11RHI/D3D11RHI_HoloLens.build.cs
new file mode 100644
--- /dev/null
+++ b/Engine/Platforms/Hololens/Source/Runtime/Windows/D3D11RHI/D3D11RHI_HoloLens.build.cs
@@ -0,0 +1,14 @@
+// Copyright Epic Games, Inc. All Rights Reserved.
+
+namespace UnrealBuildTool.Rules
+{
+	public class D3D11RHI_HoloLens : D3D11RHI
+	{
+		protected override bool bIncludeExtensions { get => false; }
+
+		public D3D11RHI_HoloLens(ReadOnlyTargetRules Target) : base(Target)
+		{
+			PrivateIncludePaths.Add("../Platforms/HoloLens/Source/Runtime/Windows/D3D11RHI/Private");
+		}
+	}
+}
