move HoloLens runtime source to platform extension folder.

diff --git a/Engine/Platforms/Hololens/Source/Runtime/Windows/D3D11RHI/Private/D3D11RHIPrivateUtil.h b/Engine/Platforms/Hololens/Source/Runtime/Windows/D3D11RHI/Private/D3D11RHIPrivateUtil.h
new file mode 100644
--- /dev/null
+++ b/Engine/Platforms/Hololens/Source/Runtime/Windows/D3D11RHI/Private/D3D11RHIPrivateUtil.h
@@ -0,0 +1,24 @@
+// Copyright Epic Games, Inc. All Rights Reserved.
+
+/*=============================================================================
+	D3D11RHIPrivateUtil.h: Private D3D RHI Utility definitions for HoloLens.
+=============================================================================*/
+
+#pragma once
+
+#include "HoloLensD3D11ConstantBuffer.h"
+
+struct FD3DRHIUtil
+{
+	template <EShaderFrequency ShaderFrequencyT>
+	static FORCEINLINE void CommitConstants(FD3D11ConstantBuffer* InConstantBuffer, FD3D11StateCache& StateCache, uint32 Index, bool bDiscardSharedConstants)
+	{
+		auto* ConstantBuffer = ((FWinD3D11ConstantBuffer*)InConstantBuffer);
+		// Array may contain NULL entries to pad out to proper 
+		if (ConstantBuffer && ConstantBuffer->CommitConstantsToDevice(bDiscardSharedConstants))
+		{
+			ID3D11Buffer* DeviceBuffer = ConstantBuffer->GetConstantBuffer();
+			StateCache.SetConstantBuffer<ShaderFrequencyT>(DeviceBuffer, Index);
+		}
+	}
+};
diff --git a/Engine/Platforms/Hololens/Source/Runtime/Windows/D3D11RHI/Private/HoloLensD3D11ConstantBuffer.cpp b/Engine/Platforms/Hololens/Source/Runtime/Windows/D3D11RHI/Private/HoloLensD3D11ConstantBuffer.cpp
new file mode 100644
--- /dev/null
+++ b/Engine/Platforms/Hololens/Source/Runtime/Windows/D3D11RHI/Private/HoloLensD3D11ConstantBuffer.cpp
@@ -0,0 +1,126 @@
+// Copyright Epic Games, Inc. All Rights Reserved.
+
+/*=============================================================================
+	HoloLensD3D11ConstantBuffer.cpp: D3D Constant Buffer functions
+=============================================================================*/
+
+#include "D3D11RHIPrivate.h"
+#include "HoloLensD3D11ConstantBuffer.h"
+
+void FWinD3D11ConstantBuffer::InitDynamicRHI()
+{
+	TRefCountPtr<ID3D11Buffer> CBuffer = NULL;
+	D3D11_BUFFER_DESC BufferDesc;
+	ZeroMemory( &BufferDesc, sizeof( D3D11_BUFFER_DESC ) );
+	// Verify constant buffer ByteWidth requirements
+	check(MaxSize <= D3D11_REQ_CONSTANT_BUFFER_ELEMENT_COUNT && (MaxSize % 16) == 0);
+	BufferDesc.ByteWidth = MaxSize;
+
+	BufferDesc.Usage = D3D11_USAGE_DEFAULT;
+	BufferDesc.CPUAccessFlags = 0;
+	BufferDesc.BindFlags = D3D11_BIND_CONSTANT_BUFFER;
+	BufferDesc.MiscFlags = 0;
+
+	Buffers = new TRefCountPtr<ID3D11Buffer>[NumSubBuffers];
+	CurrentSubBuffer = 0;
+	for(uint32 s = 0;s < NumSubBuffers;s++)
+	{
+		VERIFYD3D11RESULT(D3DRHI->GetDevice()->CreateBuffer(&BufferDesc, NULL, Buffers[s].GetInitReference()));
+		UpdateBufferStats(Buffers[s],true);
+		BufferDesc.ByteWidth = Align(BufferDesc.ByteWidth / 2, 16);
+	}
+
+	FD3D11ConstantBuffer::InitDynamicRHI();
+}
+
+void FWinD3D11ConstantBuffer::ReleaseDynamicRHI()
+{
+	FD3D11ConstantBuffer::ReleaseDynamicRHI();
+
+	if(Buffers)
+	{
+		for(uint32 s = 0;s < NumSubBuffers;s++)
+		{
+			UpdateBufferStats(Buffers[s],false);
+		}
+
+		delete [] Buffers;
+		Buffers = NULL;
+	}
+}
+
+bool FWinD3D11ConstantBuffer::CommitConstantsToDevice( bool bDiscardSharedConstants )
+{
+	// New circular buffer system for faster constant uploads.  Avoids CopyResource and speeds things up considerably
+	if (CurrentUpdateSize == 0)
+	{
+		return false;
+	}
+
+	SCOPE_CYCLE_COUNTER(STAT_D3D11GlobalConstantBufferUpdateTime);
+
+	if ( bDiscardSharedConstants )
+	{
+		// If we're discarding shared constants, just use constants that have been updated since the last Commit.
+		TotalUpdateSize = CurrentUpdateSize;
+	}
+	else
+	{
+		// If we're re-using shared constants, use all constants.
+		TotalUpdateSize = FMath::Max( CurrentUpdateSize, TotalUpdateSize );
+	}
+
+	// This basically keeps track dynamically how much data has been updated every frame
+	// and then divides up a "max" constant buffer size by halves down until it finds a large sections that more tightly covers
+	// the amount updated, assuming that all data in a constant buffer is updated each draw call and contiguous.
+	// This only used for the IndexSlot==0 constant buffer on the vertex shader.
+	// Which will have an indeterminate number of constant values that are generated by material shaders.
+	CurrentSubBuffer = 1;
+	uint32 BufferSize = MaxSize / 2;
+	while ( BufferSize >= TotalUpdateSize && CurrentSubBuffer < NumSubBuffers )
+	{
+		CurrentSubBuffer++;
+		BufferSize /= 2;
+	}
+	CurrentSubBuffer--;
+	BufferSize *= 2;
+
+	ID3D11Buffer* Buffer = Buffers[CurrentSubBuffer];
+	check(((uint64)(ShadowData) & 0xf) == 0);
+	D3DRHI->GetDeviceContext()->UpdateSubresource(Buffer,0,NULL,(void*)ShadowData,BufferSize,BufferSize);
+
+	CurrentUpdateSize = 0;
+
+	return true;
+}
+
+void FD3D11DynamicRHI::InitConstantBuffers()
+{
+	// Allocate shader constant buffers.  All shader types can have access to all buffers.
+	//  Index==0 is reserved for "custom" params, and the rest are reserved by the system for Common
+	//	constants
+	VSConstantBuffers.Empty(MAX_CONSTANT_BUFFER_SLOTS);
+	PSConstantBuffers.Empty(MAX_CONSTANT_BUFFER_SLOTS);
+	GSConstantBuffers.Empty(MAX_CONSTANT_BUFFER_SLOTS);
+	CSConstantBuffers.Empty(MAX_CONSTANT_BUFFER_SLOTS);
+	for(int32 BufferIndex = 0;BufferIndex < MAX_CONSTANT_BUFFER_SLOTS;BufferIndex++)
+	{
+		uint32 Size = GConstantBufferSizes[BufferIndex];
+		uint32 SubBuffers = 1;
+		if(BufferIndex == GLOBAL_CONSTANT_BUFFER_INDEX)
+		{
+			SubBuffers = 5;
+		}
+
+		// Vertex shader can have subbuffers for index==0.  This is from Epic's original design for the auto-fit of size to
+		//	reduce the update costs of the buffer.
+
+		// New circular buffer system for faster constant uploads.  Avoids CopyResource and speeds things up considerably
+		VSConstantBuffers.Add(new FWinD3D11ConstantBuffer(this, Size, SubBuffers));
+		PSConstantBuffers.Add(new FWinD3D11ConstantBuffer(this, Size, SubBuffers));
+		GSConstantBuffers.Add(new FWinD3D11ConstantBuffer(this, Size));
+		CSConstantBuffers.Add(new FWinD3D11ConstantBuffer(this, Size));
+	}
+}
+
+DEFINE_STAT(STAT_D3D11GlobalConstantBufferUpdateTime);
diff --git a/Engine/Platforms/Hololens/Source/Runtime/Windows/D3D11RHI/Private/HoloLensD3D11ConstantBuffer.h b/Engine/Platforms/Hololens/Source/Runtime/Windows/D3D11RHI/Private/HoloLensD3D11ConstantBuffer.h
new file mode 100644
--- /dev/null
+++ b/Engine/Platforms/Hololens/Source/Runtime/Windows/D3D11RHI/Private/HoloLensD3D11ConstantBuffer.h
@@ -0,0 +1,41 @@
+// Copyright Epic Games, Inc. All Rights Reserved.
+
+/*=============================================================================
+	HoloLensD3D11ConstantBuffer.h: D3D Constant Buffer functions
+=============================================================================*/
+
+#pragma once
+
+class FWinD3D11ConstantBuffer : public FD3D11ConstantBuffer
+{
+public:
+	FWinD3D11ConstantBuffer(FD3D11DynamicRHI* InD3DRHI, uint32 InSize = 0, uint32 SubBuffers = 1) :
+		FD3D11ConstantBuffer(InD3DRHI, InSize, SubBuffers),
+		Buffers(nullptr),
+		CurrentSubBuffer(0),
+		NumSubBuffers(SubBuffers)
+	{
+	}
+
+	// FRenderResource interface.
+	virtual void	InitDynamicRHI() override;
+	virtual void	ReleaseDynamicRHI() override;
+
+	/**
+	* Get the current pool buffer
+	*/
+	ID3D11Buffer* GetConstantBuffer() const
+	{
+		return Buffers[CurrentSubBuffer];
+	}
+
+	/**
+	* Unlocks the constant buffer so the data can be transmitted to the device
+	*/
+	bool CommitConstantsToDevice(bool bDiscardSharedConstants);
+
+private:
+	TRefCountPtr<ID3D11Buffer>* Buffers;
+	uint32	CurrentSubBuffer;
+	uint32	NumSubBuffers;
+};
diff --git a/Engine/Platforms/Hololens/Source/Runtime/Windows/D3D11RHI/Private/HoloLensD3D11Device.cpp b/Engine/Platforms/Hololens/Source/Runtime/Windows/D3D11RHI/Private/HoloLensD3D11Device.cpp
new file mode 100644
--- /dev/null
+++ b/Engine/Platforms/Hololens/Source/Runtime/Windows/D3D11RHI/Private/HoloLensD3D11Device.cpp
@@ -0,0 +1,10 @@
+// Copyright Epic Games, Inc. All Rights Reserved.
+
+/*=============================================================================
+	HoloLensD3D11Device.cpp: Windows D3D device RHI implementation.
+=============================================================================*/
+
+#include "D3D11RHIPrivate.h"
+
+// Borrow the Windows desktop version with unavailable APIs remapped
+#include "Windows/WindowsD3D11Device.cpp"
diff --git a/Engine/Platforms/Hololens/Source/Runtime/Windows/D3D11RHI/Private/HoloLensD3D11Viewport.cpp b/Engine/Platforms/Hololens/Source/Runtime/Windows/D3D11RHI/Private/HoloLensD3D11Viewport.cpp
new file mode 100644
--- /dev/null
+++ b/Engine/Platforms/Hololens/Source/Runtime/Windows/D3D11RHI/Private/HoloLensD3D11Viewport.cpp
@@ -0,0 +1,242 @@
+// Copyright Epic Games, Inc. All Rights Reserved.
+
+#include "D3D11RHIPrivate.h"
+#include "RenderCore.h"
+
+#include "Microsoft/AllowMicrosoftPlatformTypes.h"
+	#include <DirectXMath.h>
+#include "Microsoft/HideMicrosoftPlatformTypes.h"
+
+#pragma warning(disable : 4946)	// reinterpret_cast used between related classes: 'Platform::Object' and ...
+
+using namespace Windows::Foundation;
+using namespace Windows::ApplicationModel::Core;
+using namespace Windows::UI::Core;
+
+uint32 FD3D11Viewport::GSwapChainFlags = DXGI_SWAP_CHAIN_FLAG_ALLOW_MODE_SWITCH;
+
+DXGI_FORMAT GetSupportedSwapChainBufferFormat(DXGI_FORMAT InPreferredDXGIFormat)
+{
+	switch (InPreferredDXGIFormat)
+	{
+	case DXGI_FORMAT_B8G8R8A8_UNORM:
+	case DXGI_FORMAT_R8G8B8A8_UNORM:
+	case DXGI_FORMAT_R16G16B16A16_FLOAT:
+		return InPreferredDXGIFormat;
+
+	// More sophisticated fallbacks here?
+
+	default:
+		return DXGI_FORMAT_B8G8R8A8_UNORM;
+	}
+}
+
+// FD3D11Viewport::GetBackBufferFormat()
+//{
+//	return DXGI_FORMAT_B8G8R8A8_UNORM;
+//}
+
+FD3D11Viewport::FD3D11Viewport(FD3D11DynamicRHI* InD3DRHI,HWND InWindowHandle,uint32 InSizeX,uint32 InSizeY,bool bInIsFullscreen, EPixelFormat InPreferredPixelFormat):
+	D3DRHI(InD3DRHI),
+	LastFlipTime(0),
+	LastFrameComplete(0),
+	LastCompleteTime(0),
+	SyncCounter(0),
+	bSyncedLastFrame(false),
+	WindowHandle(InWindowHandle),
+	SizeX(InSizeX),
+	SizeY(InSizeY),
+	PresentFailCount(0),
+	ValidState(0),
+	PixelFormat(InPreferredPixelFormat),
+	bIsFullscreen(bInIsFullscreen),
+	bAllowTearing(false),
+	FrameSyncEvent(InD3DRHI)
+{
+	D3DRHI->Viewports.Add(this);
+
+	// Ensure that the D3D device has been created.
+	D3DRHI->InitD3DDevice();
+
+	// Create a backbuffer/swapchain for each viewport
+	TRefCountPtr<IDXGIDevice> DXGIDevice;
+	VERIFYD3D11RESULT_EX(D3DRHI->GetDevice()->QueryInterface(IID_IDXGIDevice, (void**)DXGIDevice.GetInitReference()), D3DRHI->GetDevice());
+
+	uint32 DisplayIndex = D3DRHI->GetHDRDetectedDisplayIndex();
+	bForcedFullscreenDisplay = FParse::Value(FCommandLine::Get(), TEXT("FullscreenDisplay="), DisplayIndex);
+
+	if (bForcedFullscreenDisplay || GRHISupportsHDROutput)
+	{
+		TRefCountPtr<IDXGIAdapter> DXGIAdapter;
+		DXGIDevice->GetAdapter((IDXGIAdapter**)DXGIAdapter.GetInitReference());
+
+		if (S_OK != DXGIAdapter->EnumOutputs(DisplayIndex, ForcedFullscreenOutput.GetInitReference()))
+		{
+			UE_LOG(LogD3D11RHI, Log, TEXT("Failed to find requested output display (%i)."), DisplayIndex);
+			ForcedFullscreenOutput = nullptr;
+			bForcedFullscreenDisplay = false;
+		}
+	}
+	else
+	{
+		ForcedFullscreenOutput = nullptr;
+	}
+
+	if (PixelFormat == PF_FloatRGBA && bIsFullscreen)
+	{
+		// Send HDR meta data to enable
+		D3DRHI->EnableHDR();
+	}
+
+	// Skip swap chain creation in off-screen rendering mode
+	bNeedSwapChain = !FParse::Param(FCommandLine::Get(), TEXT("RenderOffScreen"));
+	if (bNeedSwapChain)
+	{
+		// Create the swapchain.
+		if (InD3DRHI->IsQuadBufferStereoEnabled())
+		{
+			IDXGIFactory2* Factory2 = (IDXGIFactory2*)D3DRHI->GetFactory();
+
+			BOOL stereoEnabled = Factory2->IsWindowedStereoEnabled();
+			if (stereoEnabled)
+			{
+				DXGI_SWAP_CHAIN_DESC1 SwapChainDesc1;
+				FMemory::Memzero(&SwapChainDesc1, sizeof(DXGI_SWAP_CHAIN_DESC1));
+
+				// Enable stereo 
+				SwapChainDesc1.Stereo = true;
+				// MSAA Sample count
+				SwapChainDesc1.SampleDesc.Count = 1;
+				SwapChainDesc1.SampleDesc.Quality = 0;
+
+				SwapChainDesc1.Format = DXGI_FORMAT_R10G10B10A2_UNORM;
+				SwapChainDesc1.BufferUsage = DXGI_USAGE_RENDER_TARGET_OUTPUT | DXGI_USAGE_SHADER_INPUT;
+				// Double buffering required to create stereo swap chain
+				SwapChainDesc1.BufferCount = 2;
+				SwapChainDesc1.Scaling = DXGI_SCALING_NONE;
+				SwapChainDesc1.SwapEffect = DXGI_SWAP_EFFECT_FLIP_DISCARD;
+				SwapChainDesc1.Flags = GSwapChainFlags = DXGI_SWAP_CHAIN_FLAG_ALLOW_MODE_SWITCH;
+
+				IDXGISwapChain1* SwapChain1 = nullptr;
+				VERIFYD3D11RESULT_EX((Factory2->CreateSwapChainForCoreWindow(
+					D3DRHI->GetDevice(), 
+					reinterpret_cast<IUnknown*>(CoreWindow::GetForCurrentThread()),
+					&SwapChainDesc1, 
+					nullptr, 
+					&SwapChain1)
+					), D3DRHI->GetDevice());
+				SwapChain = SwapChain1;
+				GRHISupportsHDROutput =
+					(SwapChainDesc1.Format == DXGI_FORMAT_R10G10B10A2_TYPELESS) ||
+					(SwapChainDesc1.Format == DXGI_FORMAT_R16G16B16A16_FLOAT);
+			}
+			else
+			{
+				UE_LOG(LogD3D11RHI, Log, TEXT("FD3D11Viewport::FD3D11Viewport was not able to create stereo SwapChain; Please enable stereo in driver settings."));
+				InD3DRHI->DisableQuadBufferStereo();
+			}
+		}
+
+		// if stereo was not activated or not enabled in settings
+		if (SwapChain == nullptr)
+		{
+			IDXGIAdapter* pdxgiAdapter;
+			VERIFYD3D11RESULT(DXGIDevice->GetAdapter(&pdxgiAdapter));
+
+			IDXGIFactory2* pdxgiFactory;
+			VERIFYD3D11RESULT(pdxgiAdapter->GetParent(__uuidof(pdxgiFactory), reinterpret_cast<void**>(&pdxgiFactory)));
+
+			// Create the swapchain.
+			DXGI_SWAP_CHAIN_DESC1 swapChainDesc = { 0 };
+			swapChainDesc.Width = SizeX;
+			swapChainDesc.Height = SizeY;
+			swapChainDesc.Format = GetSupportedSwapChainBufferFormat(GetRenderTargetFormat(PixelFormat));
+			swapChainDesc.Stereo = false;
+			swapChainDesc.SampleDesc.Count = 1;                          // don't use multi-sampling
+			swapChainDesc.SampleDesc.Quality = 0;
+			swapChainDesc.BufferUsage = DXGI_USAGE_RENDER_TARGET_OUTPUT | DXGI_USAGE_SHADER_INPUT;
+			swapChainDesc.BufferCount = 2;                               // use two buffers to enable flip effect
+			swapChainDesc.Scaling = DXGI_SCALING_NONE;
+			swapChainDesc.AlphaMode = DXGI_ALPHA_MODE_IGNORE;
+			swapChainDesc.SwapEffect = DXGI_SWAP_EFFECT_FLIP_SEQUENTIAL; // we recommend using this swap effect for all applications
+			swapChainDesc.Flags = GSwapChainFlags = 0;
+
+			IDXGISwapChain1* CreatedSwapChain = NULL;
+
+			VERIFYD3D11RESULT(pdxgiFactory->CreateSwapChainForCoreWindow(
+				D3DRHI->GetDevice(),
+				reinterpret_cast<IUnknown*>(CoreWindow::GetForCurrentThread()),
+				&swapChainDesc,
+				NULL,
+				(IDXGISwapChain1**)&CreatedSwapChain
+			)
+			);
+
+			*(SwapChain.GetInitReference()) = CreatedSwapChain;
+
+			GRHISupportsHDROutput =
+				(swapChainDesc.Format == DXGI_FORMAT_R10G10B10A2_TYPELESS) ||
+				(swapChainDesc.Format == DXGI_FORMAT_R16G16B16A16_FLOAT);
+
+		}
+	}
+
+	// Create a RHI surface to represent the viewport's back buffer.
+	BackBuffer = GetSwapChainSurface(D3DRHI, PixelFormat, SizeX, SizeY, SwapChain);
+
+	ENQUEUE_RENDER_COMMAND(FD3D11Viewport)(
+	[this](FRHICommandListImmediate& RHICmdList)
+	{
+		// Initialize the query by issuing an initial event.
+		FrameSyncEvent.IssueEvent();
+	});
+}
+
+
+// When a window has moved or resized we need to check whether it is on a HDR monitor or not. Set the correct color space of the monitor
+void FD3D11Viewport::CheckHDRMonitorStatus()
+{
+	PixelColorSpace = EColorSpaceAndEOTF::ERec709_sRGB;
+}
+
+void FD3D11Viewport::ConditionalResetSwapChain(bool bIgnoreFocus)
+{
+	uint32 Valid = ValidState;
+	if (0 != (Valid & VIEWPORT_INVALID))
+	{
+		if (0 != (Valid & VIEWPORT_FULLSCREEN_LOST))
+		{
+			FlushRenderingCommands();
+			ValidState &= ~(VIEWPORT_FULLSCREEN_LOST);
+			Resize(SizeX, SizeY, false, PixelFormat);
+		}
+		else
+		{
+			ResetSwapChainInternal(bIgnoreFocus);
+		}
+	}
+}
+
+void FD3D11Viewport::ResetSwapChainInternal(bool bIgnoreFocus)
+{
+	if (0 != (ValidState & VIEWPORT_INVALID))
+	{
+		const bool bIsFocused = true;
+		const bool bIsIconic = false;
+		if(bIgnoreFocus || (bIsFocused && !bIsIconic) )
+		{
+			FlushRenderingCommands();
+
+			HRESULT Result = SwapChain->SetFullscreenState(bIsFullscreen,NULL);
+			if(SUCCEEDED(Result))
+			{
+				ValidState &= ~(VIEWPORT_INVALID);
+			}
+			else
+			{
+				// Even though the docs say SetFullscreenState always returns S_OK, that doesn't always seem to be the case.
+				UE_LOG(LogD3D11RHI, Log, TEXT("IDXGISwapChain::SetFullscreenState returned %08x; waiting for the next frame to try again."),Result);
+			}
+		}
+	}
+}
\ No newline at end of file
diff --git a/Engine/Plugins/Media/HLMedia/Source/HLMedia/HLMedia.Build.cs b/Engine/Plugins/Media/HLMedia/Source/HLMedia/HLMedia.Build.cs
--- a/Engine/Plugins/Media/HLMedia/Source/HLMedia/HLMedia.Build.cs
+++ b/Engine/Plugins/Media/HLMedia/Source/HLMedia/HLMedia.Build.cs
@@ -69,7 +69,7 @@ public class HLMedia : ModuleRules
         }
         else if (Target.Platform == UnrealTargetPlatform.HoloLens)
         {
-            PrivateIncludePaths.Add(Path.Combine(EngineDir, "Source/Runtime/Windows/D3D11RHI/Private/HoloLens"));
+            PrivateIncludePaths.Add(Path.Combine(EngineDir, "Platforms/HoloLens/Source/Runtime/Windows/D3D11RHI/Private"));
         }
     }
 }
diff --git a/Engine/Plugins/Runtime/WindowsMixedReality/Source/WindowsMixedRealityHMD/Private/WindowsMixedRealityHMD.h b/Engine/Plugins/Runtime/WindowsMixedReality/Source/WindowsMixedRealityHMD/Private/WindowsMixedRealityHMD.h
--- a/Engine/Plugins/Runtime/WindowsMixedReality/Source/WindowsMixedRealityHMD/Private/WindowsMixedRealityHMD.h
+++ b/Engine/Plugins/Runtime/WindowsMixedReality/Source/WindowsMixedRealityHMD/Private/WindowsMixedRealityHMD.h
@@ -13,7 +13,7 @@
 #include "Windows/WindowsApplication.h"
 #endif
 #if PLATFORM_HOLOLENS
-#include "HoloLens/HoloLensApplication.h"
+#include "HoloLensApplication.h"
 #endif
 
 #include "Framework/Application/SlateApplication.h"
diff --git a/Engine/Plugins/Runtime/WindowsMixedReality/Source/WindowsMixedRealityRHI/WindowsMixedRealityRHI.Build.cs b/Engine/Plugins/Runtime/WindowsMixedReality/Source/WindowsMixedRealityRHI/WindowsMixedRealityRHI.Build.cs
--- a/Engine/Plugins/Runtime/WindowsMixedReality/Source/WindowsMixedRealityRHI/WindowsMixedRealityRHI.Build.cs
+++ b/Engine/Plugins/Runtime/WindowsMixedReality/Source/WindowsMixedRealityRHI/WindowsMixedRealityRHI.Build.cs
@@ -92,7 +92,7 @@ namespace UnrealBuildTool.Rules
                 }
                 else if (Target.Platform == UnrealTargetPlatform.HoloLens)
                 {
-                    PrivateIncludePaths.Add("../../../../Source/Runtime/Windows/D3D11RHI/Private/HoloLens");
+                    PrivateIncludePaths.Add("../../../../Platforms/HoloLens/Source/Runtime/Windows/D3D11RHI/Private");
                 }
 
                 PCHUsage = PCHUsageMode.NoSharedPCHs;
diff --git a/Engine/Source/Developer/HoloLens/HoloLensTargetPlatform/Private/HoloLensTargetPlatform.h b/Engine/Source/Developer/HoloLens/HoloLensTargetPlatform/Private/HoloLensTargetPlatform.h
--- a/Engine/Source/Developer/HoloLens/HoloLensTargetPlatform/Private/HoloLensTargetPlatform.h
+++ b/Engine/Source/Developer/HoloLens/HoloLensTargetPlatform/Private/HoloLensTargetPlatform.h
@@ -7,7 +7,7 @@
 #pragma once
 
 #include "Common/TargetPlatformBase.h"
-#include "Runtime/Core/Public/HoloLens/HoloLensPlatformProperties.h"
+#include "HoloLensPlatformProperties.h"
 #include "Misc/ConfigCacheIni.h"
 #include "HoloLensTargetDevice.h"
 #include "Misc/ScopeLock.h"
diff --git a/Engine/Source/Runtime/Core/Core.Build.cs b/Engine/Source/Runtime/Core/Core.Build.cs
--- a/Engine/Source/Runtime/Core/Core.Build.cs
+++ b/Engine/Source/Runtime/Core/Core.Build.cs
@@ -93,17 +93,6 @@ public class Core : ModuleRules
 				PublicDefinitions.Add("UE_MEMORY_TAGS_TRACE_ENABLED=1");
 			}
 		}
-		else if ((Target.Platform == UnrealTargetPlatform.HoloLens))
-		{
-			PublicIncludePaths.Add("Runtime/Core/Public/HoloLens");
-			AddEngineThirdPartyPrivateStaticDependencies(Target,
-				"zlib");
-
-			AddEngineThirdPartyPrivateStaticDependencies(Target,
-				"IntelTBB",
-				"XInput"
-				);
-		}
 		else if (Target.Platform == UnrealTargetPlatform.Mac)
 		{
 			AddEngineThirdPartyPrivateStaticDependencies(Target,
@@ -213,11 +202,6 @@ public class Core : ModuleRules
 			}
 		}
 
-		if (Target.Platform == UnrealTargetPlatform.HoloLens)
-		{
-			PublicDefinitions.Add("WITH_VS_PERF_PROFILER=0");
-		}
-
 		// Superluminal instrumentation support, if one has it installed
 		if (Target.Platform.IsInGroup(UnrealPlatformGroup.Windows))
 		{
@@ -279,7 +263,7 @@ public class Core : ModuleRules
 		}
 
 		if ((Target.Platform == UnrealTargetPlatform.Mac) || (Target.Platform == UnrealTargetPlatform.IOS) || (Target.Platform == UnrealTargetPlatform.TVOS) 
-			|| (Target.Platform == UnrealTargetPlatform.HoloLens) || (Target.Platform == UnrealTargetPlatform.Android))
+			|| (Target.Platform == UnrealTargetPlatform.Android))
 		{
 			PublicDefinitions.Add("IS_RUNNING_GAMETHREAD_ON_EXTERNAL_THREAD=1");
 		}
diff --git a/Engine/Source/Runtime/Core/Public/HAL/ExceptionHandling.h b/Engine/Source/Runtime/Core/Public/HAL/ExceptionHandling.h
--- a/Engine/Source/Runtime/Core/Public/HAL/ExceptionHandling.h
+++ b/Engine/Source/Runtime/Core/Public/HAL/ExceptionHandling.h
@@ -41,7 +41,7 @@ extern CORE_API void ReportEnsure(const TCHAR* ErrorMessage, void* ProgramCounte
 extern CORE_API void ReportStall(const TCHAR* ErrorMessage, uint32 ThreadId);
 extern CORE_API void ReportHang(const TCHAR*, const uint64* StackFrames, int32 NumStackFrames, uint32 HungThreadId);
 #elif PLATFORM_HOLOLENS
-#include "HoloLens/HoloLensSystemIncludes.h"
+#include "HoloLensSystemIncludes.h"
 #include <excpt.h>
 extern CORE_API int32 ReportCrash(Windows::LPEXCEPTION_POINTERS ExceptionInfo);
 extern CORE_API void ReportEnsure(const TCHAR* ErrorMessage, void* ProgramCounter);
diff --git a/Engine/Source/Runtime/Core/Public/Microsoft/AllowMicrosoftPlatformAtomics.h b/Engine/Source/Runtime/Core/Public/Microsoft/AllowMicrosoftPlatformAtomics.h
--- a/Engine/Source/Runtime/Core/Public/Microsoft/AllowMicrosoftPlatformAtomics.h
+++ b/Engine/Source/Runtime/Core/Public/Microsoft/AllowMicrosoftPlatformAtomics.h
@@ -2,8 +2,6 @@
 
 #if PLATFORM_WINDOWS
 	#include "Windows/AllowWindowsPlatformAtomics.h"
-#elif PLATFORM_HOLOLENS
-	#include "HoloLens/AllowWindowsPlatformAtomics.h"
 #else
 	#include "Microsoft/AllowMicrosoftPlatformAtomicsPrivate.h"
 #endif
diff --git a/Engine/Source/Runtime/Core/Public/Microsoft/AllowMicrosoftPlatformTypes.h b/Engine/Source/Runtime/Core/Public/Microsoft/AllowMicrosoftPlatformTypes.h
--- a/Engine/Source/Runtime/Core/Public/Microsoft/AllowMicrosoftPlatformTypes.h
+++ b/Engine/Source/Runtime/Core/Public/Microsoft/AllowMicrosoftPlatformTypes.h
@@ -4,8 +4,6 @@
 
 #if PLATFORM_WINDOWS
 	#include "Windows/AllowWindowsPlatformTypes.h"
-#elif PLATFORM_HOLOLENS
-	#include "HoloLens/AllowWindowsPlatformTypes.h"
 #else
 	#include "Microsoft/AllowMicrosoftPlatformTypesPrivate.h"
 #endif
diff --git a/Engine/Source/Runtime/Core/Public/Microsoft/HideMicrosoftPlatformAtomics.h b/Engine/Source/Runtime/Core/Public/Microsoft/HideMicrosoftPlatformAtomics.h
--- a/Engine/Source/Runtime/Core/Public/Microsoft/HideMicrosoftPlatformAtomics.h
+++ b/Engine/Source/Runtime/Core/Public/Microsoft/HideMicrosoftPlatformAtomics.h
@@ -2,8 +2,6 @@
 
 #if PLATFORM_WINDOWS
 	#include "Windows/HideWindowsPlatformAtomics.h"
-#elif PLATFORM_HOLOLENS
-	#include "HoloLens/HideWindowsPlatformAtomics.h"
 #else
 	#include "Microsoft/HideMicrosoftPlatformAtomicsPrivate.h"
 #endif
diff --git a/Engine/Source/Runtime/Core/Public/Microsoft/HideMicrosoftPlatformTypes.h b/Engine/Source/Runtime/Core/Public/Microsoft/HideMicrosoftPlatformTypes.h
--- a/Engine/Source/Runtime/Core/Public/Microsoft/HideMicrosoftPlatformTypes.h
+++ b/Engine/Source/Runtime/Core/Public/Microsoft/HideMicrosoftPlatformTypes.h
@@ -2,8 +2,6 @@
 
 #if PLATFORM_WINDOWS
 	#include "Windows/HideWindowsPlatformTypes.h"
-#elif PLATFORM_HOLOLENS
-	#include "HoloLens/HideWindowsPlatformTypes.h"
 #else
 	#include "Microsoft/HideMicrosoftPlatformTypesPrivate.h"
 #endif
diff --git a/Engine/Source/Runtime/Core/Public/Microsoft/WindowsHWrapper.h b/Engine/Source/Runtime/Core/Public/Microsoft/WindowsHWrapper.h
--- a/Engine/Source/Runtime/Core/Public/Microsoft/WindowsHWrapper.h
+++ b/Engine/Source/Runtime/Core/Public/Microsoft/WindowsHWrapper.h
@@ -5,8 +5,6 @@
 
 #if PLATFORM_WINDOWS
 	#include "Windows/WindowsHWrapper.h"
-#elif PLATFORM_HOLOLENS
-	#include "HoloLens/WindowsHWrapper.h"
 #else
     #include "Microsoft/WindowsHWrapperPrivate.h"
 #endif
diff --git a/Engine/Source/Runtime/D3D12RHI/D3D12RHI.Build.cs b/Engine/Source/Runtime/D3D12RHI/D3D12RHI.Build.cs
--- a/Engine/Source/Runtime/D3D12RHI/D3D12RHI.Build.cs
+++ b/Engine/Source/Runtime/D3D12RHI/D3D12RHI.Build.cs
@@ -4,12 +4,10 @@ using UnrealBuildTool;
 
 public class D3D12RHI : ModuleRules
 {
+	protected virtual bool bUsesWindowsD3D12 { get => false; }
+
 	public D3D12RHI(ReadOnlyTargetRules Target) : base(Target)
 	{
-		if (Target.Platform == UnrealTargetPlatform.HoloLens)
-		{
-			PrivateIncludePaths.Add("Runtime/D3D12RHI/Private/HoloLens");
-		}
 		PrivateIncludePaths.Add("Runtime/D3D12RHI/Private");
 		PrivateIncludePaths.Add("../Shaders/Shared");
 
@@ -41,7 +39,7 @@ public class D3D12RHI : ModuleRules
 		AddEngineThirdPartyPrivateStaticDependencies(Target, "IntelExtensionsFramework");
 
         if (Target.Platform.IsInGroup(UnrealPlatformGroup.Windows) ||
-            Target.Platform == UnrealTargetPlatform.HoloLens)
+			bUsesWindowsD3D12)
 		{
 			PublicDefinitions.Add("D3D12RHI_PLATFORM_HAS_CUSTOM_INTERFACE=0");
 
@@ -56,7 +54,7 @@ public class D3D12RHI : ModuleRules
 				PublicDependencyModuleNames.Add("WinPixEventRuntime");
 			}
 
-			if (Target.Platform != UnrealTargetPlatform.HoloLens)
+			if (Target.Platform.IsInGroup(UnrealPlatformGroup.Windows))
             {
 				PrivateDependencyModuleNames.Add("GeForceNOWWrapper");
 				AddEngineThirdPartyPrivateStaticDependencies(Target, "NVAPI");
diff --git a/Engine/Source/Runtime/D3D12RHI/Private/D3D12RHIPrivate.h b/Engine/Source/Runtime/D3D12RHI/Private/D3D12RHIPrivate.h
--- a/Engine/Source/Runtime/D3D12RHI/Private/D3D12RHIPrivate.h
+++ b/Engine/Source/Runtime/D3D12RHI/Private/D3D12RHIPrivate.h
@@ -41,8 +41,6 @@ DECLARE_LOG_CATEGORY_EXTERN(LogD3D12GapRecorder, Log, All);
 
 #if PLATFORM_WINDOWS
 #include "Windows/D3D12RHIBasePrivate.h"
-#elif PLATFORM_HOLOLENS
-#include "HoloLens/D3D12RHIBasePrivate.h"
 #else
 #include "D3D12RHIBasePrivate.h"
 #endif
diff --git a/Engine/Source/Runtime/D3D12RHI/Public/ID3D12DynamicRHI.h b/Engine/Source/Runtime/D3D12RHI/Public/ID3D12DynamicRHI.h
--- a/Engine/Source/Runtime/D3D12RHI/Public/ID3D12DynamicRHI.h
+++ b/Engine/Source/Runtime/D3D12RHI/Public/ID3D12DynamicRHI.h
@@ -7,8 +7,6 @@
 
 #if PLATFORM_WINDOWS
 	#include "Windows/D3D12ThirdParty.h"
-#elif PLATFORM_HOLOLENS
-	#include "HoloLens/D3D12ThirdParty.h"
 #else
 	#include "D3D12ThirdParty.h"
 #endif
diff --git a/Engine/Source/Runtime/Launch/Launch.Build.cs b/Engine/Source/Runtime/Launch/Launch.Build.cs
--- a/Engine/Source/Runtime/Launch/Launch.Build.cs
+++ b/Engine/Source/Runtime/Launch/Launch.Build.cs
@@ -83,11 +83,6 @@ public class Launch : ModuleRules
 					"AudioMixerXAudio2",
 				});
 			}
-			else if (Target.Platform == UnrealTargetPlatform.HoloLens)
-			{
-				DynamicallyLoadedModuleNames.Add("D3D11RHI");
-				DynamicallyLoadedModuleNames.Add("AudioMixerXAudio2");
-			}
 			else if (Target.Platform == UnrealTargetPlatform.Mac)
 			{
 				DynamicallyLoadedModuleNames.AddRange(new string [] {
diff --git a/Engine/Source/Runtime/Online/HTTP/Private/IXML/HttpIXMLSupport.cpp b/Engine/Source/Runtime/Online/HTTP/Private/IXML/HttpIXMLSupport.cpp
--- a/Engine/Source/Runtime/Online/HTTP/Private/IXML/HttpIXMLSupport.cpp
+++ b/Engine/Source/Runtime/Online/HTTP/Private/IXML/HttpIXMLSupport.cpp
@@ -8,7 +8,7 @@
 #include <string>
 #include <sstream>
 
-#include "AllowWindowsPlatformTypes.h"
+#include "Microsoft/AllowMicrosoftPlatformTypes.h"
 
 // --------------------------------------------------------------------------------------
 // Name: HttpCallback::~HttpCallback
@@ -516,6 +516,6 @@ HRESULT RequestStream::Invoke(
 
 // --------------------------------------------------------------------------------------
 
-#include "HideWindowsPlatformTypes.h"
+#include "Microsoft/HideMicrosoftPlatformTypes.h"
 
 #endif
diff --git a/Engine/Source/Runtime/Online/HTTP/Private/IXML/HttpIXMLSupport.h b/Engine/Source/Runtime/Online/HTTP/Private/IXML/HttpIXMLSupport.h
--- a/Engine/Source/Runtime/Online/HTTP/Private/IXML/HttpIXMLSupport.h
+++ b/Engine/Source/Runtime/Online/HTTP/Private/IXML/HttpIXMLSupport.h
@@ -8,7 +8,7 @@
 
 #include "Containers/UnrealString.h"
 #include "Http.h"
-#include "HoloLens/AllowWindowsPlatformTypes.h"
+#include "Microsoft/AllowMicrosoftPlatformTypes.h"
 #include <msxml6.h>
 #include <wrl.h>
 #include <Windows.h>
@@ -111,6 +111,6 @@ protected:
 };
 
 
-#include "HoloLens/HideWindowsPlatformTypes.h"
+#include "Microsoft/HideMicrosoftPlatformTypes.h"
 
 #endif
diff --git a/Engine/Source/Runtime/Online/WebSockets/Private/PlatformWebSocket.h b/Engine/Source/Runtime/Online/WebSockets/Private/PlatformWebSocket.h
--- a/Engine/Source/Runtime/Online/WebSockets/Private/PlatformWebSocket.h
+++ b/Engine/Source/Runtime/Online/WebSockets/Private/PlatformWebSocket.h
@@ -13,7 +13,7 @@
 #elif WITH_WINHTTPWEBSOCKETS
 	#include "WinHttp/WinHttpWebSocketsManager.h"
 #elif PLATFORM_HOLOLENS
-	#include "HoloLens/HoloLensWebSocketsManager.h"
+	#include "HoloLensWebSocketsManager.h"
 #else
 	#error "Web Sockets not implemented on this platform yet"
 #endif
diff --git a/Engine/Source/Runtime/Online/WebSockets/WebSockets.Build.cs b/Engine/Source/Runtime/Online/WebSockets/WebSockets.Build.cs
--- a/Engine/Source/Runtime/Online/WebSockets/WebSockets.Build.cs
+++ b/Engine/Source/Runtime/Online/WebSockets/WebSockets.Build.cs
@@ -26,6 +26,11 @@ public class WebSockets : ModuleRules
 		}
 	}
 
+	protected virtual bool bPlatformSupportsWinRTWebsockets
+	{
+		get => false;
+	}
+
 	protected virtual bool UsePlatformSSL
 	{
 		get => false;
@@ -35,8 +40,6 @@ public class WebSockets : ModuleRules
 	{
 		get
 		{
-			bool bPlatformSupportsWinRTWebsockets = Target.Platform == UnrealTargetPlatform.HoloLens;
-
 			return PlatformSupportsLibWebsockets || bPlatformSupportsWinRTWebsockets || bPlatformSupportsWinHttpWebSockets;
 		}
 	}
diff --git a/Engine/Source/Runtime/OodleDataCompression/OodleDataCompression.Build.cs b/Engine/Source/Runtime/OodleDataCompression/OodleDataCompression.Build.cs
--- a/Engine/Source/Runtime/OodleDataCompression/OodleDataCompression.Build.cs
+++ b/Engine/Source/Runtime/OodleDataCompression/OodleDataCompression.Build.cs
@@ -46,23 +46,6 @@ using UnrealBuildTool;
 				DebugLib = "oo2core_win64_debug.lib";
 				PlatformDir = "Win64";
 			}
-			else if (Target.Platform == UnrealTargetPlatform.HoloLens)
-			{
-				if (Target.WindowsPlatform.Architecture == WindowsArchitecture.x64) // emulation target, bBuildForEmulation
-				{
-					ReleaseLib = "oo2core_winuwp64.lib";
-					DebugLib = "oo2core_winuwp64_debug.lib";
-				}
-				else if (Target.WindowsPlatform.Architecture == WindowsArchitecture.ARM64) // device target, bBuildForDevice
-				{
-					ReleaseLib = "oo2core_winuwparm64.lib";
-					DebugLib = "oo2core_winuwparm64_debug.lib";
-				}
-				else
-				{
-					throw new System.Exception("Unknown architecture for HoloLens platform!");
-				}
-			}
 			else if (Target.Platform == UnrealTargetPlatform.Mac)
 			{
 				ReleaseLib = "liboo2coremac64.a";
diff --git a/Engine/Source/Runtime/RHI/RHI.Build.cs b/Engine/Source/Runtime/RHI/RHI.Build.cs
--- a/Engine/Source/Runtime/RHI/RHI.Build.cs
+++ b/Engine/Source/Runtime/RHI/RHI.Build.cs
@@ -48,12 +48,6 @@ public class RHI : ModuleRules
 					DynamicallyLoadedModuleNames.Add("D3D12RHI");
 				}
 
-				if ((Target.Platform == UnrealTargetPlatform.HoloLens))
-				{
-					DynamicallyLoadedModuleNames.Add("D3D11RHI");
-					DynamicallyLoadedModuleNames.Add("D3D12RHI");
-				}
-
 				if ((Target.Platform.IsInGroup(UnrealPlatformGroup.Windows)) ||
 					(Target.IsInPlatformGroup(UnrealPlatformGroup.Unix) && (Target.Architecture.StartsWith("x86_64") || Target.Architecture.StartsWith("aarch64"))))	// temporary, not all archs can support Vulkan atm
 				{
diff --git a/Engine/Source/Runtime/Windows/D3D11RHI/D3D11RHI.Build.cs b/Engine/Source/Runtime/Windows/D3D11RHI/D3D11RHI.Build.cs
--- a/Engine/Source/Runtime/Windows/D3D11RHI/D3D11RHI.Build.cs
+++ b/Engine/Source/Runtime/Windows/D3D11RHI/D3D11RHI.Build.cs
@@ -2,15 +2,13 @@
 
 using UnrealBuildTool;
 
-[SupportedPlatforms("Win64", "HoloLens")]
+[SupportedPlatforms("Win64")]
 public class D3D11RHI : ModuleRules
 {
+	protected virtual bool bIncludeExtensions { get => true; }
+
 	public D3D11RHI(ReadOnlyTargetRules Target) : base(Target)
 	{
-		if (Target.Platform == UnrealTargetPlatform.HoloLens)
-		{
-			PrivateIncludePaths.Add("Runtime/Windows/D3D11RHI/Private/HoloLens");
-		}
 		PrivateIncludePaths.Add("Runtime/Windows/D3D11RHI/Private");
 		PrivateIncludePaths.Add("../Shaders/Shared");
 
@@ -26,7 +24,7 @@ public class D3D11RHI : ModuleRules
 			);
 
 		AddEngineThirdPartyPrivateStaticDependencies(Target, "DX11");
-		if (Target.Platform != UnrealTargetPlatform.HoloLens)
+		if (bIncludeExtensions)
 		{ 
         	AddEngineThirdPartyPrivateStaticDependencies(Target, "NVAPI");
 			AddEngineThirdPartyPrivateStaticDependencies(Target, "AMD_AGS");
diff --git a/Engine/Source/Runtime/Windows/D3D11RHI/Public/ID3D11DynamicRHI.h b/Engine/Source/Runtime/Windows/D3D11RHI/Public/ID3D11DynamicRHI.h
--- a/Engine/Source/Runtime/Windows/D3D11RHI/Public/ID3D11DynamicRHI.h
+++ b/Engine/Source/Runtime/Windows/D3D11RHI/Public/ID3D11DynamicRHI.h
@@ -5,8 +5,8 @@
 
 #if PLATFORM_WINDOWS
 	#include "Windows/D3D11ThirdParty.h"
-#elif PLATFORM_HOLOLENS
-	#include "HoloLens/D3D11ThirdParty.h"
+#else
+	#include "D3D11ThirdParty.h"
 #endif
 
 typedef ID3D11DeviceContext FD3D11DeviceContext;
