Delete all .target files (arm64, and x64 as well, rather than just multi)

diff --git a/Engine/Platforms/Hololens/Source/Programs/AutomationTool/HoloLensPlatform.Automation.cs b/Engine/Platforms/Hololens/Source/Programs/AutomationTool/HoloLensPlatform.Automation.cs
--- a/Engine/Platforms/Hololens/Source/Programs/AutomationTool/HoloLensPlatform.Automation.cs
+++ b/Engine/Platforms/Hololens/Source/Programs/AutomationTool/HoloLensPlatform.Automation.cs
@@ -381,10 +381,22 @@ namespace HoloLens.Automation
 			{
 				foreach(var target in Params.ClientCookedTargets)
 				{
-					var ReceiptFileName = TargetReceipt.GetDefaultPath(Params.RawProjectPath.Directory, target, UnrealTargetPlatform.HoloLens, BuildConfig, "Multi");
-					if(File.Exists(ReceiptFileName.FullName))
+					var MultiReceiptFileName = TargetReceipt.GetDefaultPath(Params.RawProjectPath.Directory, target, UnrealTargetPlatform.HoloLens, BuildConfig, "Multi");
+					if (File.Exists(MultiReceiptFileName.FullName))
 					{
-						FileReference.Delete(ReceiptFileName);
+						FileReference.Delete(MultiReceiptFileName);
+					}
+
+					var x64ReceiptFileName = TargetReceipt.GetDefaultPath(Params.RawProjectPath.Directory, target, UnrealTargetPlatform.HoloLens, BuildConfig, "x64");
+					if (File.Exists(x64ReceiptFileName.FullName))
+					{
+						FileReference.Delete(x64ReceiptFileName);
+					}
+
+					var arm64ReceiptFileName = TargetReceipt.GetDefaultPath(Params.RawProjectPath.Directory, target, UnrealTargetPlatform.HoloLens, BuildConfig, "arm64");
+					if (File.Exists(arm64ReceiptFileName.FullName))
+					{
+						FileReference.Delete(arm64ReceiptFileName);
 					}
 
 					foreach (var Arch in ActualArchitectures)
