Fixing case of //UE5/Main/Engine/Platforms/HoloLens/Source/Runtime/Core/... (2/2)

#rnx

diff --git a/Engine/Platforms/Hololens/Source/Runtime/Core/Private/HoloLensPlatformMisc.cpp b/Engine/Platforms/Hololens/Source/Runtime/Core/Private/HoloLensPlatformMisc.cpp
new file mode 100644
--- /dev/null
+++ b/Engine/Platforms/Hololens/Source/Runtime/Core/Private/HoloLensPlatformMisc.cpp
@@ -0,0 +1,425 @@
+// Copyright Epic Games, Inc. All Rights Reserved.
+
+/*=============================================================================
+	HoloLensMisc.cpp: HoloLens implementations of misc functions
+=============================================================================*/
+
+#include "HoloLensPlatformMisc.h"
+#include "Internationalization/Internationalization.h"
+#include "Misc/MessageDialog.h"
+#include "Misc/App.h"
+#include "HAL/ExceptionHandling.h"
+#include "Misc/SecureHash.h"
+#include "Microsoft/AllowMicrosoftPlatformTypes.h"
+#include <time.h>
+#include <agile.h>
+#include "Microsoft/HideMicrosoftPlatformTypes.h"
+#include "Modules/ModuleManager.h"
+
+/** 
+ * Whether support for integrating into the firewall is there
+ */
+#define WITH_FIREWALL_SUPPORT	0
+
+/** Original C- Runtime pure virtual call handler that is being called in the (highly likely) case of a double fault */
+_purecall_handler DefaultPureCallHandler;
+
+/**
+ * Our own pure virtual function call handler, set by appPlatformPreInit. Falls back
+ * to using the default C- Runtime handler in case of double faulting.
+ */ 
+static void PureCallHandler()
+{
+	static bool bHasAlreadyBeenCalled = false;
+	UE_DEBUG_BREAK();
+	if (bHasAlreadyBeenCalled)
+	{
+		// Call system handler if we're double faulting.
+		if (DefaultPureCallHandler)
+		{
+			DefaultPureCallHandler();
+		}
+	}
+	else
+	{
+		bHasAlreadyBeenCalled = true;
+		if (GIsRunning)
+		{
+			FMessageDialog::Open(EAppMsgType::Ok, NSLOCTEXT("Core", "PureVirtualFunctionCalledWhileRunningApp", "Pure virtual function being called while application was running (GIsRunning == 1)."));
+		}
+		UE_LOG(LogTemp, Fatal, TEXT("Pure virtual function being called"));
+	}
+}
+
+/*-----------------------------------------------------------------------------
+	SHA-1 functions.
+-----------------------------------------------------------------------------*/
+
+/**
+ * Get the hash values out of the executable hash section
+ *
+ * NOTE: hash keys are stored in the executable, you will need to put a line like the
+ *		 following into your PCLaunch.rc settings:
+ *			ID_HASHFILE RCDATA "../../../../GameName/Build/Hashes.sha"
+ *
+ *		 Then, use the -sha option to the cooker (must be from commandline, not
+ *       frontend) to generate the hashes for .ini, loc, startup packages, and .usf shader files
+ *
+ *		 You probably will want to make and checkin an empty file called Hashses.sha
+ *		 into your source control to avoid linker warnings. Then for testing or final
+ *		 final build ONLY, use the -sha command and relink your executable to put the
+ *		 hashes for the current files into the executable.
+ */
+static void InitSHAHashes()
+{
+	[REDACTED_EMAIL]: Implement this using GetModuleSection?
+}
+
+void FHoloLensMisc::PlatformPreInit()
+{
+	FGenericPlatformMisc::PlatformPreInit();
+
+	// Use our own handler for pure virtuals being called.
+	DefaultPureCallHandler = _set_purecall_handler(PureCallHandler);
+
+	// initialize the file SHA hash mapping
+	InitSHAHashes();
+
+	PumpMessages(true);
+}
+
+void FHoloLensMisc::PlatformInit()
+{
+	// Identity.
+	UE_LOG(LogInit, Log, TEXT("Computer: %s"), FPlatformProcess::ComputerName());
+	UE_LOG(LogInit, Log, TEXT("User: %s"), FPlatformProcess::UserName());
+
+	// Timer resolution.
+	UE_LOG(LogInit, Log, TEXT("High frequency timer resolution =%f MHz"), 0.000001 / FPlatformTime::GetSecondsPerCycle());
+
+	PumpMessages(true);
+}
+
+void FHoloLensMisc::PlatformPostInit(bool ShowSplashScreen)
+{
+	PumpMessages(true);
+}
+
+bool FHoloLensMisc::CoInitialize()
+{
+	HRESULT hr = ::CoInitializeEx(NULL, COINIT_APARTMENTTHREADED);
+	return hr == S_OK || hr == S_FALSE;
+}
+
+void FHoloLensMisc::CoUninitialize()
+{
+	::CoUninitialize();
+}
+
+FString FHoloLensMisc::GetEnvironmentVariable(const TCHAR* VariableName)
+{
+	[REDACTED_EMAIL]:
+	return FString();
+}
+
+const TCHAR* FHoloLensMisc::GetPlatformFeaturesModuleName()
+{
+	// Expectation is that if we hand back a non-null name then it's guaranteed to load.
+	// Since the platform features supported on HoloLens are strictly optional and enabled via
+	// plug-in we'll do a quick pre-check here.
+	static const TCHAR* PlatformFeaturesModuleName = TEXT("HoloLensPlatformFeatures");
+	IModuleInterface* CustomModule = FModuleManager::LoadModulePtr<IModuleInterface>(PlatformFeaturesModuleName);
+	return CustomModule ? PlatformFeaturesModuleName : nullptr;
+}
+
+// Defined in HoloLensLaunch.cpp
+extern void appWinPumpMessages();
+
+void FHoloLensMisc::PumpMessages(bool bFromMainLoop)
+{
+	if (!bFromMainLoop)
+	{
+		// Process pending windows messages, which is necessary to the rendering thread in some rare cases where DX10
+		// sends window messages (from IDXGISwapChain::Present) to the main thread owned viewport window.
+		// Only process sent messages to minimize the opportunity for re-entry in the editor, since wx messages are not deferred.
+		return;
+	}
+
+	// Handle all incoming messages if we're not using wxWindows in which case this is done by their
+	// message loop.
+	appWinPumpMessages();
+
+	// check to see if the window in the foreground was made by this process (ie, does this app
+	// have focus)
+	[REDACTED_EMAIL]: Will this always be true?
+	bool HasFocus = true;
+	// if its our window, allow sound, otherwise apply multiplier
+	FApp::SetVolumeMultiplier(HasFocus ? 1.0f : FApp::GetUnfocusedVolumeMultiplier());
+}
+
+void FHoloLensMisc::LowLevelOutputDebugString(const TCHAR *Message)
+{
+	OutputDebugString(Message);
+}
+
+void FHoloLensMisc::RequestExit(bool Force)
+{
+	UE_LOG(LogTemp, Log, TEXT("FHoloLensMisc::RequestExit(%i)"), Force);
+	if (Force)
+	{
+		// Force immediate exit.
+		// Dangerous because config code isn't flushed, global destructors aren't called, etc.
+		//::TerminateProcess(GetCurrentProcess(), 0); 
+		RequestEngineExit(TEXT("FHoloLensMisc::RequestExit"));
+	}
+	else
+	{
+		RequestEngineExit(TEXT("FHoloLensMisc::RequestExit"));
+	}
+}
+
+const TCHAR* FHoloLensMisc::GetSystemErrorMessage(TCHAR* OutBuffer, int32 BufferCount, int32 Error)
+{
+	check(OutBuffer && BufferCount);
+	*OutBuffer = TEXT('\0');
+	if (Error == 0)
+	{
+		Error = GetLastError();
+	}
+	FormatMessage(FORMAT_MESSAGE_FROM_SYSTEM, NULL, Error, MAKELANGID(LANG_NEUTRAL, SUBLANG_DEFAULT), OutBuffer, BufferCount, NULL);
+	TCHAR* Found = FCString::Strchr(OutBuffer,TEXT('\r'));
+	if (Found)
+	{
+		*Found = TEXT('\0');
+	}
+	Found = FCString::Strchr(OutBuffer,TEXT('\n'));
+	if (Found)
+	{
+		*Found = TEXT('\0');
+	}
+	return OutBuffer;
+}
+
+void FHoloLensMisc::CreateGuid(FGuid& Result)
+{
+	verify(CoCreateGuid((GUID*)&Result)==S_OK);
+}
+
+int32 FHoloLensMisc::NumberOfCores()
+{
+	static int32 cpus = 0;
+	if (cpus == 0)
+	{
+		ULONG CpuInfoBytes = 0;
+		GetSystemCpuSetInformation(nullptr, 0, &CpuInfoBytes, GetCurrentProcess(), 0);
+
+		TArray<byte> cpuInfoBuffer;
+		cpuInfoBuffer.Reserve(CpuInfoBytes);
+		GetSystemCpuSetInformation(reinterpret_cast<PSYSTEM_CPU_SET_INFORMATION>(cpuInfoBuffer.GetData()),
+			CpuInfoBytes, &CpuInfoBytes, GetCurrentProcess(), 0);
+
+		byte* firstCpu = cpuInfoBuffer.GetData();
+		int local_cpus = 0;
+		for (byte* currentCpu = firstCpu; currentCpu < firstCpu + CpuInfoBytes; currentCpu += reinterpret_cast<PSYSTEM_CPU_SET_INFORMATION>(currentCpu)->Size)
+		{
+			local_cpus++;
+		}
+		cpus = local_cpus;
+	}
+	return cpus;
+}
+
+#if !UE_BUILD_SHIPPING
+bool FHoloLensMisc::IsDebuggerPresent()
+{
+	return !!::IsDebuggerPresent(); 
+}
+#endif // UE_BUILD_SHIPPING
+
+/** Get the application root directory. */
+const TCHAR* FHoloLensMisc::RootDir()
+{
+	static FString Path;
+	if (Path.Len() == 0)
+	{
+		Path = Windows::ApplicationModel::Package::Current->InstalledLocation->Path->Data();
+
+		// Add trailing \\ - this is important for correct operation of FPaths::MakePathRelativeTo
+		Path += TEXT("\\");
+
+		// Switch to / so that FPaths::MakeStandardFilename will work
+		Path.ReplaceInline(TEXT("\\"), TEXT("/"), ESearchCase::CaseSensitive);
+	}
+	return *Path;
+}
+
+FString FHoloLensMisc::ProtocolActivationUri = TEXT("");
+
+void FHoloLensMisc::SetProtocolActivationUri(const FString& NewUriString)
+{
+	ProtocolActivationUri = NewUriString;
+}
+
+const FString& FHoloLensMisc::GetProtocolActivationUri()
+{
+	return ProtocolActivationUri;
+}
+
+EAppReturnType::Type FHoloLensMisc::MessageBoxExt(EAppMsgType::Type MsgType, const TCHAR* Text, const TCHAR* Caption)
+{
+	using namespace Windows::UI::Core;
+	using namespace Windows::UI::Popups;
+	using namespace Windows::Foundation;
+
+	Platform::Agile<MessageDialog> HoloLensMessageBox = Platform::Agile<MessageDialog>(ref new MessageDialog(ref new Platform::String(Text), ref new Platform::String(Caption)));
+
+	Platform::String^ CancelString;
+	Platform::String^ NoToAllString;
+	Platform::String^ NoString;
+	Platform::String^ YesToAllString;
+	Platform::String^ YesString;
+	Platform::String^ OkString;
+	Platform::String^ RetryString;
+	Platform::String^ ContinueString;
+
+	// The Localize* functions will return the Key if a dialog is presented before the config system is initialized.
+	// Instead, we use hard-coded strings if config is not yet initialized.
+	if (!GConfig)
+	{
+		CancelString = TEXT("Cancel");
+		NoToAllString = TEXT("No to All");
+		NoString = TEXT("No");
+		YesToAllString = TEXT("Yes to All");
+		YesString = TEXT("Yes");
+		OkString = TEXT("OK");
+		RetryString = TEXT("Retry");
+		ContinueString = TEXT("Continue");
+	}
+	else
+	{
+		CancelString = ref new Platform::String(*NSLOCTEXT("Core", "Cancel", "Cancel").ToString());
+		NoToAllString = ref new Platform::String(*NSLOCTEXT("Core", "NoToAll", "No to All").ToString());
+		NoString = ref new Platform::String(*NSLOCTEXT("Core", "No", "No").ToString());
+		YesToAllString = ref new Platform::String(*NSLOCTEXT("Core", "YesToAll", "Yes to All").ToString());
+		YesString = ref new Platform::String(*NSLOCTEXT("Core", "Yes", "Yes").ToString());
+		OkString = ref new Platform::String(*NSLOCTEXT("Core", "OK", "OK").ToString());
+		RetryString = ref new Platform::String(*NSLOCTEXT("Core", "Retry", "Retry").ToString());
+		ContinueString = ref new Platform::String(*NSLOCTEXT("Core", "Continue", "Continue").ToString());
+	}
+
+	switch (MsgType)
+	{
+	case EAppMsgType::YesNo:
+	{
+		HoloLensMessageBox->Commands->Append(ref new UICommand(YesString, nullptr, (Platform::Object^)(int)EAppReturnType::Yes));
+		HoloLensMessageBox->Commands->Append(ref new UICommand(NoString, nullptr, (Platform::Object^)(int)EAppReturnType::No));
+		HoloLensMessageBox->CancelCommandIndex = 1;
+		break;
+	}
+	case EAppMsgType::OkCancel:
+	{
+		HoloLensMessageBox->Commands->Append(ref new UICommand(OkString, nullptr, (Platform::Object^)(int)EAppReturnType::Ok));
+		HoloLensMessageBox->Commands->Append(ref new UICommand(CancelString, nullptr, (Platform::Object^)(int)EAppReturnType::Cancel));
+		HoloLensMessageBox->CancelCommandIndex = 1;
+		break;
+	}
+	case EAppMsgType::YesNoYesAllNoAllCancel:
+		UE_LOG(LogCore, Warning, TEXT("MessageBox type requires more buttons than can be displayed on this platform.  Using fallback."));
+	case EAppMsgType::YesNoCancel:
+	{
+		HoloLensMessageBox->Commands->Append(ref new UICommand(YesString, nullptr, (Platform::Object^)(int)EAppReturnType::Yes));
+		HoloLensMessageBox->Commands->Append(ref new UICommand(NoString, nullptr, (Platform::Object^)(int)EAppReturnType::No));
+		HoloLensMessageBox->Commands->Append(ref new UICommand(CancelString, nullptr, (Platform::Object^)(int)EAppReturnType::Cancel));
+		HoloLensMessageBox->CancelCommandIndex = 2;
+		break;
+	}
+	case EAppMsgType::CancelRetryContinue:
+	{
+		HoloLensMessageBox->Commands->Append(ref new UICommand(CancelString, nullptr, (Platform::Object^)(int)EAppReturnType::Cancel));
+		HoloLensMessageBox->Commands->Append(ref new UICommand(RetryString, nullptr, (Platform::Object^)(int)EAppReturnType::Retry));
+		HoloLensMessageBox->Commands->Append(ref new UICommand(ContinueString, nullptr, (Platform::Object^)(int)EAppReturnType::Continue));
+		HoloLensMessageBox->CancelCommandIndex = 0;
+		break;
+	}
+	case EAppMsgType::YesNoYesAllNoAll:
+		UE_LOG(LogCore, Warning, TEXT("MessageBox type requires more buttons than can be displayed on this platform.  Using fallback."));
+	case EAppMsgType::YesNoYesAll:
+	{
+		HoloLensMessageBox->Commands->Append(ref new UICommand(YesString, nullptr, (Platform::Object^)(int)EAppReturnType::Yes));
+		HoloLensMessageBox->Commands->Append(ref new UICommand(YesToAllString, nullptr, (Platform::Object^)(int)EAppReturnType::YesAll));
+		HoloLensMessageBox->Commands->Append(ref new UICommand(NoString, nullptr, (Platform::Object^)(int)EAppReturnType::No));
+		HoloLensMessageBox->CancelCommandIndex = 2;
+		break;
+	}
+	case EAppMsgType::Ok:
+	default:
+	{
+		HoloLensMessageBox->Commands->Append(ref new UICommand(OkString, nullptr, (Platform::Object^)(int)EAppReturnType::Ok));
+		break;
+	}
+	}
+
+	IAsyncOperation<IUICommand^>^ UIOperation;
+	bool IsUIThread = CoreWindow::GetForCurrentThread() != nullptr;
+	if (IsUIThread)
+	{
+		UIOperation = HoloLensMessageBox->ShowAsync();
+	}
+	else
+	{
+		// Must invoke the box from the UI thread, so dispatch to that
+		IAsyncAction^ RunAction = Windows::ApplicationModel::Core::CoreApplication::MainView->Dispatcher->RunAsync(
+			Windows::UI::Core::CoreDispatcherPriority::Normal,
+			ref new Windows::UI::Core::DispatchedHandler([&UIOperation, HoloLensMessageBox]()
+		{
+			UIOperation = HoloLensMessageBox->ShowAsync();
+		}));
+
+		// Wait for the UI thread to invoke the box.  Note this completes when the
+		// box is shown, not when it is closed.
+		while (RunAction->Status == AsyncStatus::Started)
+		{
+			PumpMessages(false);
+
+			// Yield to make sure UI thread can kick in and pop the UI.
+			FHoloLensProcess::Sleep(0.f);
+		}
+
+		if (!UIOperation)
+		{
+			// Error in the dispatched call.  Default to cancel.
+			return EAppReturnType::Cancel;
+		}
+	}
+
+	// Pump the core window messages until the dialog has been dismissed.
+	while (UIOperation->Status == AsyncStatus::Started)
+	{
+		PumpMessages(IsUIThread);
+
+		// Yield - possibly needed if we're not the UI thread, and won't hurt if we are.
+		FHoloLensProcess::Sleep(0.f);
+	}
+
+	// Return the command type pressed (or cancel for error).
+	EAppReturnType::Type ReturnValue = EAppReturnType::Cancel;
+	if (UIOperation->Status == AsyncStatus::Completed)
+	{
+		ReturnValue = static_cast<EAppReturnType::Type>((int)UIOperation->GetResults()->Id);
+	}
+	return ReturnValue;
+}
+
+void FHoloLensMisc::GetValidTargetPlatforms(class TArray<class FString>& TargetPlatformNames)
+{
+	TargetPlatformNames.Add(FPlatformProperties::PlatformName());
+}
+
+bool FHoloLensMisc::VerifyWindowsVersion(uint32 MajorVersion, uint32 MinorVersion, uint32 BuildNumber)
+{
+	// HoloLens is a Windows 10 OS.  
+	// UWP does not have an equivalent of Win32's VerifyVersionInfo function which is used in WindowsPlatformMisc.
+	// The closest alternative is parsing Windows.System.Profile.AnalyticsInfo.VersionInfo.DeviceFamilyVersion,
+	// But, this string is not intended for runtime version checking and its format could change.
+	return MajorVersion <= 10;
+}
diff --git a/Engine/Platforms/Hololens/Source/Runtime/Core/Private/HoloLensPlatformProcess.cpp b/Engine/Platforms/Hololens/Source/Runtime/Core/Private/HoloLensPlatformProcess.cpp
new file mode 100644
--- /dev/null
+++ b/Engine/Platforms/Hololens/Source/Runtime/Core/Private/HoloLensPlatformProcess.cpp
@@ -0,0 +1,464 @@
+// Copyright Epic Games, Inc. All Rights Reserved.
+
+/*=============================================================================
+HoloLensProcess.cpp: HoloLens implementations of Process functions
+=============================================================================*/
+
+#include "HoloLensPlatformProcess.h"
+#include "Misc/SingleThreadEvent.h"
+#include "HoloLensRunnableThread.h"
+#include "Misc/CommandLine.h"
+#include "Misc/CoreStats.h"
+#include "Misc/Paths.h"
+#include "Microsoft/WindowsHWrapper.h"
+
+#include "Microsoft/AllowMicrosoftPlatformTypes.h"
+
+// static variables
+TArray<FString> FHoloLensProcess::DllDirectoryStack;
+
+const TCHAR* FHoloLensProcess::BaseDir()
+{
+	static bool bFirstTime = true;
+	static TCHAR Result[512]=TEXT("");
+	if (!Result[0] && bFirstTime)
+	{
+		// Check the commandline for -BASEDIR=<base directory>
+		const TCHAR* CmdLine = NULL;
+		if (FCommandLine::IsInitialized())
+		{
+			CmdLine = FCommandLine::Get();
+		}
+		// Change from Microsoft, may require further review:
+		// code path that was hitting BaseDir in static initializers was removed
+		//                  and command line parsing was moved to earlier in launch to help avoid having 
+		//                  to have command line parsing code in multiple location.
+		//                  There's still potential for error here until Epic allows platform-subclassing the
+		//                  command line with an overridable Initialize() or BeginInitialize().
+		bool overridden = false;
+		if (CmdLine != NULL)
+		{
+			FString BaseDirToken = TEXT("-BASEDIR=");
+			FString NextToken;
+			while (FParse::Token(CmdLine, NextToken, false))
+			{
+				if (NextToken.StartsWith(BaseDirToken) == true)
+				{
+					FString BaseDir = NextToken.Right(NextToken.Len() - BaseDirToken.Len());
+					BaseDir.ReplaceInline(TEXT("\\"), TEXT("/"));
+					FCString::Strcpy(Result, *BaseDir);
+					overridden = true;
+				}
+			}
+		}
+
+		if (!overridden)
+		{
+			DWORD ResultCode = GetModuleFileName(NULL, Result, _countof(Result));
+			verify(ResultCode != 0);
+			DWORD ErrorCode = ::GetLastError();
+			verify(ErrorCode != ERROR_INSUFFICIENT_BUFFER);
+
+			// Strip out the filename and add the trailing '\'
+			FString BaseDir;
+			BaseDir = FPaths::GetPath(FString(Result));
+			BaseDir += TEXT("\\");
+			FCString::Strcpy(Result, *BaseDir);
+		}
+		bFirstTime = false;
+	}
+	return Result;
+}
+
+#include "Microsoft/HideMicrosoftPlatformTypes.h"
+
+void FHoloLensProcess::Sleep( float Seconds )
+{
+	SCOPE_CYCLE_COUNTER(STAT_Sleep);
+	FThreadIdleStats::FScopeIdle Scope;
+	::Sleep((int32)(Seconds * 1000.0));
+}
+
+void FHoloLensProcess::SleepNoStats(float Seconds)
+{
+	::Sleep((uint32)(Seconds * 1000.0));
+}
+
+void FHoloLensProcess::SleepInfinite()
+{
+	while (true)
+	{
+		::Sleep(INFINITE);
+	}
+}
+
+void FHoloLensProcess::YieldThread()
+{
+	::SwitchToThread();
+}
+
+#include "HoloLensEvent.h"
+
+FEvent* FHoloLensProcess::CreateSynchEvent(bool bIsManualReset /*= false*/)
+{
+	// Allocate the new object
+	FEvent* Event = NULL;
+	if (FPlatformProcess::SupportsMultithreading())
+	{
+		Event = new FEventHoloLens();
+	}
+	else
+	{
+		// Fake vent object.
+		Event = new FSingleThreadEvent();
+	}
+	// If the internal create fails, delete the instance and return NULL
+	PRAGMA_DISABLE_DEPRECATION_WARNINGS
+	if (!Event->Create(bIsManualReset))
+	PRAGMA_ENABLE_DEPRECATION_WARNINGS
+	{
+		delete Event;
+		Event = NULL;
+	}
+	return Event;
+}
+
+#include "Microsoft/AllowMicrosoftPlatformTypes.h"
+
+bool FEventHoloLens::Wait(uint32 WaitTime, const bool bIgnoreThreadIdleStats /*= false*/)
+{
+	SCOPE_CYCLE_COUNTER(STAT_EventWait);
+	FThreadIdleStats::FScopeIdle Scope(bIgnoreThreadIdleStats);
+	check(Event);
+
+	return WaitForSingleObjectEx(Event, WaitTime, FALSE) == WAIT_OBJECT_0;
+}
+
+FRunnableThread* FHoloLensProcess::CreateRunnableThread()
+{
+	return new FRunnableThreadHoloLens();
+}
+
+const TCHAR* FHoloLensProcess::ExecutableName(bool bRemoveExtension)
+{
+	static TCHAR Result[512] = TEXT("");
+	static TCHAR ResultWithExt[512] = TEXT("");
+	if (!Result[0])
+	{
+		// Get complete path for the executable
+		if (GetModuleFileName(NULL, Result, UE_ARRAY_COUNT(Result)) != 0)
+		{
+			// Remove all of the path information by finding the base filename
+			FString FileName = Result;
+			FString FileNameWithExt = Result;
+			FCString::Strncpy(Result, *(FPaths::GetBaseFilename(FileName)), UE_ARRAY_COUNT(Result));
+			FCString::Strncpy(ResultWithExt, *(FPaths::GetCleanFilename(FileNameWithExt)), UE_ARRAY_COUNT(ResultWithExt));
+		}
+		// If the call failed, zero out the memory to be safe
+		else
+		{
+			FMemory::Memzero(Result, sizeof(Result));
+			FMemory::Memzero(ResultWithExt, sizeof(ResultWithExt));
+		}
+	}
+
+	return (bRemoveExtension ? Result : ResultWithExt);
+}
+
+void* FHoloLensProcess::GetDllHandle(const TCHAR* Filename)
+{
+	check(Filename);
+
+	FString PackageRelativePath(Filename);
+
+	if (DllDirectoryStack.Num() > 0)
+	{
+		// If a path has been pushed on the stack use that
+		FString Path = DllDirectoryStack.Top() / PackageRelativePath;
+		return ::LoadPackagedLibrary(*Path, 0ul);
+	}
+	else
+	{
+		// Incoming paths are relative to the BaseDir, but LoadPackagedLibrary wants package relative
+		// which in UE terms is the RootDir.
+		FPaths::MakePathRelativeTo(PackageRelativePath, *(FPaths::RootDir() + TEXT("/")));
+		return ::LoadPackagedLibrary(*PackageRelativePath, 0ul);
+	}
+}
+
+void FHoloLensProcess::FreeDllHandle(void* DllHandle)
+{
+	// It is okay to call FreeLibrary on 0
+	::FreeLibrary((HMODULE)DllHandle);
+}
+
+void* FHoloLensProcess::GetDllExport(void* DllHandle, const TCHAR* ProcName)
+{
+	check(DllHandle);
+	check(ProcName);
+	return (void*)::GetProcAddress((HMODULE)DllHandle, TCHAR_TO_ANSI(ProcName));
+}
+
+void FHoloLensProcess::PushDllDirectory(const TCHAR* Directory)
+{
+	DllDirectoryStack.Push(Directory);
+}
+
+void FHoloLensProcess::PopDllDirectory(const TCHAR* Directory)
+{
+	// don't allow too many pops (indicates bad code that should be fixed, but won't kill anything, so using ensure)
+	ensureMsgf(DllDirectoryStack.Num() > 0, TEXT("Tried to PopDllDirectory too many times"));
+	// verify we are popping the top
+	checkf(DllDirectoryStack.Top() == Directory, TEXT("There was a PushDllDirectory/PopDllDirectory mismatch (Popped %s, which didn't match %s)"), *DllDirectoryStack.Top(), Directory);
+	// pop it off
+	DllDirectoryStack.Pop();
+}
+
+void FHoloLensProcess::SetCurrentWorkingDirectoryToBaseDir()
+{
+	FPlatformMisc::CacheLaunchDir();
+	verify(SetCurrentDirectoryW(BaseDir()));	// failure here usually means the ACLs got messed up, reregister the app
+}
+
+FString FHoloLensProcess::GetCurrentWorkingDirectory()
+{
+	// get the current working directory (uncached)
+	TCHAR CurrentDirectory[MAX_PATH];
+	GetCurrentDirectoryW(MAX_PATH, CurrentDirectory);
+	return CurrentDirectory;
+}
+
+const TCHAR* FHoloLensProcess::UserDir()
+{
+#if WITH_EDITOR
+	return UE_BUILD_SHIPPING ? GetLocalAppDataLowLevelPath() : GetLocalAppDataRedirectPath();
+#else
+	return GetLocalAppDataLowLevelPath();
+#endif
+}
+
+const TCHAR* FHoloLensProcess::UserSettingsDir()
+{
+	// @todo is there any reason to make this a sub-folder?
+	return UserDir();
+}
+
+const TCHAR* FHoloLensProcess::UserTempDir()
+{
+#if WITH_EDITOR
+	return UE_BUILD_SHIPPING ? GetTempAppDataLowLevelPath() : GetTempAppDataRedirectPath();
+#else
+	return GetTempAppDataLowLevelPath();
+#endif
+}
+
+const TCHAR* FHoloLensProcess::ApplicationSettingsDir()
+{
+	// This is supposed to be a writeable location that exists across multiple users.
+	// For HoloLens it will have to be user-specific.
+	return UserSettingsDir();
+}
+
+#if WIN10_SDK_VERSION < 16225
+// This is a drop in equivelent to the deprecated thread affinity APIs that most legacy code is based on.
+// As with those older APIs, it's functionality may be unexpected on machines with more than 64 cores.
+DWORD_PTR WINAPI SetThreadAffinityMask(
+	_In_ HANDLE    hThread,
+	_In_ DWORD_PTR dwThreadAffinityMask
+)
+{
+	static struct CPU_INFO {
+		int Count;
+		ULONG CpuInfoBytes;
+		uint8_t* CpuInfoBuffer;
+		PSYSTEM_CPU_SET_INFORMATION* CpuInfoPtrs;
+
+		CPU_INFO() {
+			Count = 0;
+			CpuInfoBytes = 0;
+			GetSystemCpuSetInformation(nullptr, 0, &CpuInfoBytes, GetCurrentProcess(), 0);
+
+			CpuInfoBuffer = reinterpret_cast<uint8_t*>(malloc(CpuInfoBytes));
+			if (CpuInfoBuffer)
+			{
+				GetSystemCpuSetInformation(reinterpret_cast<PSYSTEM_CPU_SET_INFORMATION>(CpuInfoBuffer),
+					CpuInfoBytes, &CpuInfoBytes, GetCurrentProcess(), 0);
+
+				BYTE* firstCpu = CpuInfoBuffer;
+				for (BYTE* currentCpu = firstCpu; currentCpu < firstCpu + CpuInfoBytes; currentCpu += reinterpret_cast<PSYSTEM_CPU_SET_INFORMATION>(currentCpu)->Size)
+				{
+					Count++;
+				}
+
+				CpuInfoPtrs = reinterpret_cast<PSYSTEM_CPU_SET_INFORMATION*>(calloc(Count, sizeof(PSYSTEM_CPU_SET_INFORMATION)));
+
+				if (CpuInfoPtrs)
+				{
+					firstCpu = CpuInfoBuffer;
+					int currentCpuNum = 0;
+					for (BYTE* currentCpu = firstCpu; currentCpu < firstCpu + CpuInfoBytes; currentCpu += reinterpret_cast<PSYSTEM_CPU_SET_INFORMATION>(currentCpu)->Size)
+					{
+						CpuInfoPtrs[currentCpuNum] = reinterpret_cast<PSYSTEM_CPU_SET_INFORMATION>(currentCpu);
+						currentCpuNum++;
+					}
+				}
+			}
+		}
+	} Cpu_Info;
+
+	// if static initialization got Cpu info....
+	if (Cpu_Info.Count > 0 && Cpu_Info.CpuInfoBytes && Cpu_Info.CpuInfoBuffer && Cpu_Info.CpuInfoPtrs)
+	{
+#ifdef _WIN64
+		DWORD_PTR priorMask = 0xffffffffffffffffull;
+		ULONG coreIds[64];
+#else
+		DWORD_PTR priorMask = 0xfffffffful;
+		ULONG coreIds[32];
+#endif
+		ULONG priorCoreCount = 0;
+
+		// this is simplified, assuming that not other setter will be setting affinity
+		if (GetThreadSelectedCpuSets(hThread, coreIds, ARRAYSIZE(coreIds), &priorCoreCount))
+		{
+			if (priorCoreCount > 0 && priorCoreCount <= ARRAYSIZE(coreIds))
+			{
+				priorMask = 0;
+				for (unsigned int currentCoreEntry = 0; currentCoreEntry < priorCoreCount; currentCoreEntry++)
+				{
+					for (int i = 0; i < Cpu_Info.Count; i++)
+					{
+						if (coreIds[currentCoreEntry] == Cpu_Info.CpuInfoPtrs[i]->CpuSet.Id)
+						{
+							priorMask |= (1ull << i);
+							break;
+						}
+					}
+				}
+			}
+		}
+
+		int markedCount = 0;
+		for (int coreNum = 0; coreNum < ARRAYSIZE(coreIds) && coreNum < Cpu_Info.Count; coreNum++)
+		{
+			if (dwThreadAffinityMask & (1ull << coreNum))
+			{
+				coreIds[markedCount] = Cpu_Info.CpuInfoPtrs[coreNum]->CpuSet.Id;
+				markedCount++;
+			}
+		}
+
+		if (SetThreadSelectedCpuSets(hThread, coreIds, markedCount))
+		{
+			return priorMask;
+		}
+		else
+		{
+			SetLastError(ERROR_BAD_ARGUMENTS);
+			return 0;
+		}
+	}
+	else
+	{
+		SetLastError(ERROR_DEVICE_ENUMERATION_ERROR);
+		return 0;
+	}
+}
+#endif
+
+void FHoloLensProcess::SetThreadAffinityMask(uint64 AffinityMask)
+{
+	::SetThreadAffinityMask(::GetCurrentThread(), (DWORD_PTR)AffinityMask);
+}
+
+const TCHAR* FHoloLensProcess::GetLocalAppDataLowLevelPath()
+{
+	static TCHAR Result[PLATFORM_MAX_FILEPATH_LENGTH] = TEXT("");
+	if (!Result[0])
+	{
+		Windows::Storage::StorageFolder ^ TempFolder = Windows::Storage::ApplicationData::Current->LocalFolder;
+		if (nullptr != TempFolder)
+		{
+			Platform::String^ TempPath = TempFolder->Path;
+			if (nullptr != TempPath)
+			{
+				FCString::Strncpy(Result, TempPath->Data(), PLATFORM_MAX_FILEPATH_LENGTH);
+			}
+		}
+
+		check(Result[0]);
+	}
+
+	return Result;
+}
+
+const TCHAR* FHoloLensProcess::GetTempAppDataLowLevelPath()
+{
+	static TCHAR Result[PLATFORM_MAX_FILEPATH_LENGTH] = TEXT("");
+	if (!Result[0])
+	{
+		Windows::Storage::StorageFolder ^ TempFolder = Windows::Storage::ApplicationData::Current->TemporaryFolder;
+		if (nullptr != TempFolder)
+		{
+			Platform::String^ TempPath = TempFolder->Path;
+			if (nullptr != TempPath)
+			{
+				FCString::Strncpy(Result, TempPath->Data(), PLATFORM_MAX_FILEPATH_LENGTH);
+			}
+		}
+
+		check(Result[0]);
+	}
+
+	return Result;
+}
+
+const TCHAR* FHoloLensProcess::GetLocalAppDataRedirectPath()
+{
+	static TCHAR Result[PLATFORM_MAX_FILEPATH_LENGTH] = TEXT("");
+	if (!Result[0])
+	{
+		FCString::Strcpy(Result, *FPaths::ConvertRelativePathToFull(FPaths::ProjectDir() / TEXT("HoloLensLocalAppData")));
+	}
+	return Result;
+}
+
+const TCHAR* FHoloLensProcess::GetTempAppDataRedirectPath()
+{
+	static TCHAR Result[PLATFORM_MAX_FILEPATH_LENGTH] = TEXT("");
+	if (!Result[0])
+	{
+		FCString::Strcpy(Result, *FPaths::ConvertRelativePathToFull(FPaths::ProjectDir() / TEXT("HoloLensTempAppData")));
+	}
+	return Result;
+}
+
+bool FHoloLensProcess::ShouldSaveToUserDir()
+{
+	return true;
+}
+
+bool FHoloLensProcess::CanLaunchURL(const TCHAR* URL)
+{
+	return URL != nullptr;
+}
+
+void FHoloLensProcess::LaunchURL(const TCHAR* URL, const TCHAR* Params, FString* Error)
+{
+#if PLATFORM_HOLOLENS
+	try
+	{
+		Windows::Foundation::Uri^ uri = ref new Windows::Foundation::Uri(ref new Platform::String(URL));
+		Windows::System::Launcher::LaunchUriAsync(uri);
+	}
+	catch (Platform::Exception^)
+	{
+		if (Error != nullptr)
+		{
+			*Error = TEXT("Failed to launch URL");
+		}
+	}
+#endif
+}
+
+#include "Microsoft/HideMicrosoftPlatformTypes.h"
diff --git a/Engine/Platforms/Hololens/Source/Runtime/Core/Private/HoloLensPlatformStackWalk.cpp b/Engine/Platforms/Hololens/Source/Runtime/Core/Private/HoloLensPlatformStackWalk.cpp
new file mode 100644
--- /dev/null
+++ b/Engine/Platforms/Hololens/Source/Runtime/Core/Private/HoloLensPlatformStackWalk.cpp
@@ -0,0 +1,15 @@
+// Copyright Epic Games, Inc. All Rights Reserved.
+
+/*=============================================================================
+	HoloLensStackWalk.cpp: HoloLens implementations of stack walk functions
+=============================================================================*/
+
+#include "HoloLensPlatformStackWalk.h"
+
+TArray<FProgramCounterSymbolInfo> FHololensPlatformStackWalk::GetStack(int32 IgnoreCount, int32 MaxDepth, void* Context)
+{
+	TArray<FProgramCounterSymbolInfo> Stack;
+	FProgramCounterSymbolInfo dummy;
+	Stack.Push(dummy);
+	return Stack;
+}
diff --git a/Engine/Platforms/Hololens/Source/Runtime/Core/Private/HoloLensPlatformTime.cpp b/Engine/Platforms/Hololens/Source/Runtime/Core/Private/HoloLensPlatformTime.cpp
new file mode 100644
--- /dev/null
+++ b/Engine/Platforms/Hololens/Source/Runtime/Core/Private/HoloLensPlatformTime.cpp
@@ -0,0 +1,124 @@
+// Copyright Epic Games, Inc. All Rights Reserved.
+
+/*=============================================================================
+HoloLensTime.cpp: HoloLens implementations of time functions
+=============================================================================*/
+
+//#include "CorePrivatePCH.h"
+#include "HoloLensPlatformTime.h"
+#include "Misc/AssertionMacros.h"
+#include "Containers/Ticker.h"
+#include "Microsoft/WindowsHWrapper.h"
+#include "Stats/Stats.h"
+
+float FHoloLensTime::CPUTimePctRelative = 0.0f;
+
+double FHoloLensTime::InitTiming()
+{
+	LARGE_INTEGER Frequency;
+	verify(QueryPerformanceFrequency(&Frequency));
+	SecondsPerCycle = 1.0 / Frequency.QuadPart;
+	SecondsPerCycle64 = 1.0 / Frequency.QuadPart;
+
+	// Due to some limitation of the OS, we limit the polling frequency to 4 times per second, 
+	// but it should be enough for longterm CPU usage monitoring.
+	static const float PollingInterval = 1.0f / 4.0f;
+
+	// Register a ticker delegate for updating the CPU utilization data.
+	FTSTicker::GetCoreTicker().AddTicker(FTickerDelegate::CreateStatic(&FPlatformTime::UpdateCPUTime), PollingInterval);
+
+	return FPlatformTime::Seconds();
+}
+
+void FHoloLensTime::SystemTime(int32& Year, int32& Month, int32& DayOfWeek, int32& Day, int32& Hour, int32& Min, int32& Sec, int32& MSec)
+{
+	SYSTEMTIME st;
+	::GetLocalTime(&st);
+
+	Year = st.wYear;
+	Month = st.wMonth;
+	DayOfWeek = st.wDayOfWeek;
+	Day = st.wDay;
+	Hour = st.wHour;
+	Min = st.wMinute;
+	Sec = st.wSecond;
+	MSec = st.wMilliseconds;
+}
+
+void FHoloLensTime::UtcTime(int32& Year, int32& Month, int32& DayOfWeek, int32& Day, int32& Hour, int32& Min, int32& Sec, int32& MSec)
+{
+	SYSTEMTIME st;
+	GetSystemTime(&st);
+
+	Year = st.wYear;
+	Month = st.wMonth;
+	DayOfWeek = st.wDayOfWeek;
+	Day = st.wDay;
+	Hour = st.wHour;
+	Min = st.wMinute;
+	Sec = st.wSecond;
+	MSec = st.wMilliseconds;
+}
+
+/** Holds Windows filetime misc functions. */
+struct FFiletimeMisc
+{
+	/**
+	 * @return number of ticks based on the specified Filetime.
+	 */
+	static FORCEINLINE uint64 TicksFromFileTime(const FILETIME& Filetime)
+	{
+		const uint64 NumTicks = (uint64(Filetime.dwHighDateTime) << 32) + Filetime.dwLowDateTime;
+		return NumTicks;
+	}
+
+	/**
+	 * @return number of seconds based on the specified Filetime.
+	 */
+	static FORCEINLINE double ToSeconds(const FILETIME& Filetime)
+	{
+		return double(TicksFromFileTime(Filetime)) / double(ETimespan::TicksPerSecond);
+	}
+};
+
+
+bool FHoloLensTime::UpdateCPUTime(float /*DeltaTime*/)
+{
+	QUICK_SCOPE_CYCLE_COUNTER(STAT_FHoloLensPlatformTime_UpdateCPUTime);
+
+	static double LastTotalProcessTime = 0.0f;
+	static double LastTotalUserAndKernelTime = 0.0f;
+
+	FILETIME CreationTime = { 0 };
+	FILETIME ExitTime = { 0 };
+	FILETIME KernelTime = { 0 };
+	FILETIME UserTime = { 0 };
+	FILETIME CurrentTime = { 0 };
+
+	::GetProcessTimes(::GetCurrentProcess(), &CreationTime, &ExitTime, &KernelTime, &UserTime);
+	::GetSystemTimeAsFileTime(&CurrentTime);
+
+	const double CurrentTotalUserAndKernelTime = FFiletimeMisc::ToSeconds(KernelTime) + FFiletimeMisc::ToSeconds(UserTime);
+	const double CurrentTotalProcessTime = FFiletimeMisc::ToSeconds(CurrentTime) - FFiletimeMisc::ToSeconds(CreationTime);
+
+	const double IntervalProcessTime = CurrentTotalProcessTime - LastTotalProcessTime;
+	const double IntervalUserAndKernelTime = CurrentTotalUserAndKernelTime - LastTotalUserAndKernelTime;
+
+	// IntervalUserAndKernelTime == 0.0f means that the OS hasn't updated the data yet, 
+	// so don't update to avoid oscillating between 0 and calculated value.
+	if (IntervalUserAndKernelTime > 0.0f)
+	{
+		CPUTimePctRelative = IntervalUserAndKernelTime / IntervalProcessTime * 100.0f;
+
+		LastTotalProcessTime = CurrentTotalProcessTime;
+		LastTotalUserAndKernelTime = CurrentTotalUserAndKernelTime;
+	}
+
+	return true;
+}
+
+
+FCPUTime FHoloLensTime::GetCPUTime()
+{
+	return FCPUTime(CPUTimePctRelative / (float)FPlatformMisc::NumberOfCoresIncludingHyperthreads(), CPUTimePctRelative);
+}
diff --git a/Engine/Platforms/Hololens/Source/Runtime/Core/Public/Microsoft/MinWindowsPrivate.h b/Engine/Platforms/Hololens/Source/Runtime/Core/Public/Microsoft/MinWindowsPrivate.h
new file mode 100644
--- /dev/null
+++ b/Engine/Platforms/Hololens/Source/Runtime/Core/Public/Microsoft/MinWindowsPrivate.h
@@ -0,0 +1,115 @@
+// Copyright Epic Games, Inc. All Rights Reserved.
+
+/*=============================================================================
+	MinHoloLensApi.h: Minimal includes for HoloLens
+=============================================================================*/
+
+#pragma once
+
+#include "CoreTypes.h"
+
+#if defined(_WINDOWS_) && !defined(UE_MINIMAL_WINDOWS_INCLUDE)
+	#pragma message ( " " )
+	#pragma message ( "You have included windows.h before MinWindows.h" )
+	#pragma message ( "All useless stuff from the windows headers won't be excluded !!!" )
+	#pragma message ( " " )
+#endif // _WINDOWS_
+
+#if !PLATFORM_HOLOLENS
+#error PLATFORM_HOLOLENS not defined
+#endif
+
+#define UE_MINIMAL_WINDOWS_INCLUDE
+
+// WIN32_LEAN_AND_MEAN excludes rarely-used services from windows headers.
+#define WIN32_LEAN_AND_MEAN
+
+// The below excludes some other unused services from the windows headers -- see windows.h for details.
+#define NOGDICAPMASKS			// CC_*, LC_*, PC_*, CP_*, TC_*, RC_
+//#define NOVIRTUALKEYCODES		// VK_*
+//#define NOWINMESSAGES			// WM_*, EM_*, LB_*, CB_*
+//#define NOWINSTYLES			// WS_*, CS_*, ES_*, LBS_*, SBS_*, CBS_*
+//#define NOSYSMETRICS			// SM_*
+#define NOMENUS					// MF_*
+//#define NOICONS				// IDI_*
+//#define NOKEYSTATES			// MK_*
+//#define NOSYSCOMMANDS			// SC_*
+//#define NORASTEROPS			// Binary and Tertiary raster ops
+//#define NOSHOWWINDOW			// SW_*
+#define OEMRESOURCE				// OEM Resource values
+#define NOATOM					// Atom Manager routines
+//#define NOCLIPBOARD			// Clipboard routines
+//#define NOCOLOR				// Screen colors
+//#define NOCTLMGR				// Control and Dialog routines
+#define NODRAWTEXT				// DrawText() and DT_*
+//#define NOGDI					// All GDI #defines and routines
+#define NOKERNEL				// All KERNEL #defines and routines
+//#define NOUSER				// All USER #defines and routines
+//#define NONLS					// All NLS #defines and routines
+//#define NOMB					// MB_* and MessageBox()
+#define NOMEMMGR				// GMEM_*, LMEM_*, GHND, LHND, associated routines
+#define NOMETAFILE				// typedef METAFILEPICT
+#define NOMINMAX				// Macros min(a,b) and max(a,b)
+//#define NOMSG					// typedef MSG and associated routines
+#define NOOPENFILE				// OpenFile(), OemToAnsi, AnsiToOem, and OF_*
+#define NOSCROLL				// SB_* and scrolling routines
+#define NOSERVICE				// All Service Controller routines, SERVICE_ equates, etc.
+#define NOSOUND					// Sound driver routines
+//#define NOTEXTMETRIC			// typedef TEXTMETRIC and associated routines
+//#define NOWH					// SetWindowsHook and WH_*
+//#define NOWINOFFSETS			// GWL_*, GCL_*, associated routines
+#define NOCOMM					// COMM driver routines
+#define NOKANJI					// Kanji support stuff.
+#define NOHELP					// Help engine interface.
+#define NOPROFILER				// Profiler interface.
+#define NODEFERWINDOWPOS		// DeferWindowPos routines
+#define NOMCX					// Modem Configuration Extensions
+#define NOCRYPT
+#define NOTAPE
+#define NOIMAGE
+#define NOPROXYSTUB
+#define NORPC
+
+#ifndef TRUE
+#define TRUE 1
+#endif
+#ifndef FALSE
+#define FALSE 0
+#endif
+
+// Finally now we can include windows.h
+#ifndef InterlockedIncrement
+#define InterlockedIncrement_MinHoloLensApi_h_hidden
+#define InterlockedIncrement _InterlockedIncrement
+#endif // !InterlockedIncrement
+
+#ifndef InterlockedDecrement
+#define InterlockedDecrement_MinHoloLensApi_h_hidden
+#define InterlockedDecrement _InterlockedDecrement
+#endif // !InterlockedDecrement
+
+#ifndef InterlockedCompareExchangePointer
+#define InterlockedCompareExchangePointer_MinHoloLensApi_h_hidden
+#if PLATFORM_64BITS
+#define InterlockedCompareExchangePointer _InterlockedCompareExchangePointer
+#else
+#define InterlockedCompareExchangePointer __InlineInterlockedCompareExchangePointer
+#endif
+#endif // !InterlockedCompareExchangePointer
+
+#include <wrl.h>
+
+#ifdef InterlockedCompareExchangePointer_MinHoloLensApi_h_hidden
+#undef InterlockedCompareExchangePointer
+#undef InterlockedCompareExchangePointer_MinHoloLensApi_h_hidden
+#endif // InterlockedCompareExchangePointer_MinHoloLensApi_h_hidden
+
+#ifdef InterlockedIncrement_MinHoloLensApi_h_hidden
+#undef InterlockedIncrement
+#undef InterlockedIncrement_MinHoloLensApi_h_hidden
+#endif // InterlockedIncrement_MinHoloLensApi_h_hidden
+
+#ifdef InterlockedDecrement_MinHoloLensApi_h_hidden
+#undef InterlockedDecrement
+#undef InterlockedDecrement_MinHoloLensApi_h_hidden
+#endif // InterlockedIncrement_MinHoloLensApi_h_hidden
diff --git a/Engine/Platforms/Hololens/Source/Runtime/Core/Public/Microsoft/PostWindowsApiPrivate.h b/Engine/Platforms/Hololens/Source/Runtime/Core/Public/Microsoft/PostWindowsApiPrivate.h
new file mode 100644
--- /dev/null
+++ b/Engine/Platforms/Hololens/Source/Runtime/Core/Public/Microsoft/PostWindowsApiPrivate.h
@@ -0,0 +1,77 @@
+// Copyright Epic Games, Inc. All Rights Reserved.
+
+// Re-enable warnings
+THIRD_PARTY_INCLUDES_START
+THIRD_PARTY_INCLUDES_END
+
+// Undo any Windows defines.
+#undef uint8
+#undef uint16
+#undef uint32
+#undef int32
+#undef float
+#undef CDECL
+#undef PF_MAX
+#undef PlaySound
+#undef DrawText
+#undef CaptureStackBackTrace
+#undef MemoryBarrier
+#undef DeleteFile
+#undef MoveFile
+#undef CopyFile
+#undef CreateDirectory
+#undef GetCurrentTime
+#undef SendMessage
+#undef LoadString
+#undef UpdateResource
+#undef FindWindow
+#undef GetObject
+#undef GetEnvironmentVariable
+#undef CreateFont
+#undef CreateDesktop
+#undef GetMessage
+#undef PostMessage
+#undef GetCommandLine
+#undef GetProp
+#undef SetPort
+#undef SetProp
+#undef GetFileAttributes
+#undef ReportEvent
+#undef GetClassName
+#undef GetClassInfo
+#undef Yield
+#undef IMediaEventSink
+#undef GetTempFileName
+#undef GetFreeSpace
+
+// Undefine all the atomics. AllowMicrosoftPlatformAtomics/HideMicrosoftPlatformAtomics temporarily defining these macros.
+#undef InterlockedIncrement
+#undef InterlockedDecrement
+#undef InterlockedAdd
+#undef InterlockedExchange
+#undef InterlockedExchangeAdd
+#undef InterlockedCompareExchange
+#undef InterlockedCompareExchangePointer
+#undef InterlockedExchange64
+#undef InterlockedExchangeAdd64
+#undef InterlockedCompareExchange64
+#undef InterlockedIncrement64
+#undef InterlockedDecrement64
+#undef InterlockedAnd
+#undef InterlockedOr
+#undef InterlockedXor
+
+// Restore any previously defined macros
+#pragma pop_macro("MAX_uint8")
+#pragma pop_macro("MAX_uint16")
+#pragma pop_macro("MAX_uint32")
+#pragma pop_macro("MAX_int32")
+#pragma pop_macro("TEXT")
+
+// Redefine CDECL to our version of the #define.  <AJS> Is this really necessary?
+#define CDECL	    __cdecl					/* Standard C function */
+
+// Make sure version is high enough for API to be defined. For CRITICAL_SECTION
+#if !defined(_XTL_) && (_WIN32_WINNT < 0x0403)
+#error SetCriticalSectionSpinCount requires _WIN32_WINNT >= 0x0403
+#endif
diff --git a/Engine/Platforms/Hololens/Source/Runtime/Core/Public/Microsoft/PreWindowsApiPrivate.h b/Engine/Platforms/Hololens/Source/Runtime/Core/Public/Microsoft/PreWindowsApiPrivate.h
new file mode 100644
--- /dev/null
+++ b/Engine/Platforms/Hololens/Source/Runtime/Core/Public/Microsoft/PreWindowsApiPrivate.h
@@ -0,0 +1,24 @@
+// Copyright Epic Games, Inc. All Rights Reserved.
+
+/*=============================================================================
+	PreHoloLensApi.h: The pre-include part of UE4's Win32 API wrapper.
+=============================================================================*/
+
+// Disable the warning that the pack size is changed in this header.
+#pragma warning(disable:4103)
+
+// Save these macros for later; Windows redefines them
+#pragma push_macro("MAX_uint8")
+#pragma push_macro("MAX_uint16")
+#pragma push_macro("MAX_uint32")
+#pragma push_macro("MAX_int32")
+#pragma push_macro("TEXT")
+
+// Undefine the TEXT macro for winnt.h to redefine it, unless it's already been included
+#ifndef _WINNT_
+#undef TEXT
+#endif
+
+// Disable all normal third party headers
+THIRD_PARTY_INCLUDES_START
+THIRD_PARTY_INCLUDES_END
\ No newline at end of file
