UBT: Use ILogger throughout UBT, rather than legacy EpicGames.Core.Log methods. All output going forwards should use structured logging rather than String.Format style logging.

diff --git a/Engine/Platforms/Hololens/Source/Programs/AutomationTool/HoloLensPlatform.Automation.cs b/Engine/Platforms/Hololens/Source/Programs/AutomationTool/HoloLensPlatform.Automation.cs
--- a/Engine/Platforms/Hololens/Source/Programs/AutomationTool/HoloLensPlatform.Automation.cs
+++ b/Engine/Platforms/Hololens/Source/Programs/AutomationTool/HoloLensPlatform.Automation.cs
@@ -512,7 +512,7 @@ namespace HoloLens.Automation
 				Windows10SDKVersion = "Latest";
 			}
 
-			if(!HoloLensExports.InitWindowsSdkToolPath(Windows10SDKVersion))
+			if(!HoloLensExports.InitWindowsSdkToolPath(Windows10SDKVersion, Log.Logger))
 			{
 				throw new AutomationException(ExitCode.Error_Arguments, "Wrong WinSDK toolchain selected on \'Platforms/HoloLens/Toolchain\' page. Please check.");
 			}
@@ -596,7 +596,7 @@ namespace HoloLens.Automation
 			{
 				// Stage all the build products
 				DirectoryReference ProjectBinariesFolder = Params.GetProjectBinariesPathForPlatform(PlatformType);
-				HoloLensExports DeployExports = new HoloLensExports();
+				HoloLensExports DeployExports = new HoloLensExports(Log.Logger);
 				foreach (StageTarget Target in SC.StageTargets)
 				{
 					string ArchString = Target.Receipt.Architecture;
@@ -725,7 +725,7 @@ namespace HoloLens.Automation
 		private void PackagePakFiles(ProjectParams Params, DeploymentContext SC, string OutputNameBase)
 		{
 			string IntermediateDirectory = Path.Combine(SC.ProjectRoot.FullName, "Intermediate", "Deploy", "neutral");
-			var ListResources = new HoloLensManifestGenerator().CreateAssetsManifest(SC.StageTargetPlatform.PlatformType, SC.StageDirectory.FullName, IntermediateDirectory, SC.RawProjectPath, SC.ProjectRoot.FullName);
+			var ListResources = new HoloLensManifestGenerator(Log.Logger).CreateAssetsManifest(SC.StageTargetPlatform.PlatformType, SC.StageDirectory.FullName, IntermediateDirectory, SC.RawProjectPath, SC.ProjectRoot.FullName);
 
 			string OutputName = OutputNameBase + "_pak";
 
@@ -1236,7 +1236,7 @@ namespace HoloLens.Automation
 			StartInfo.Arguments = String.Format("\"{0}\" \"{1}\" -p", SourceFile.FullName, TargetFile.FullName);
 			StartInfo.UseShellExecute = false;
 			StartInfo.CreateNoWindow = true;
-			Utils.RunLocalProcessAndLogOutput(StartInfo);
+			Utils.RunLocalProcessAndLogOutput(StartInfo, Log.Logger);
 
 			if (bStripInPlace)
 			{
@@ -1751,7 +1751,7 @@ namespace HoloLens.Automation
 		private string[] GetPathToVCLibsPackages(bool UseDebugCrt, WindowsCompiler Compiler)
 		{
 			List<DirectoryReference> SdkRootDirs = new List<DirectoryReference>();
-			WindowsExports.EnumerateSdkRootDirs(SdkRootDirs);
+			WindowsExports.EnumerateSdkRootDirs(SdkRootDirs, Log.Logger);
 			
 			string VCVersionFragment;
             switch (Compiler)
@@ -1839,7 +1839,7 @@ namespace HoloLens.Automation
 			}
 
 
-			HoloLensExports.CreateManifestForDLC(Params.DLCFile, SC.StageDirectory);
+			HoloLensExports.CreateManifestForDLC(Params.DLCFile, SC.StageDirectory, Log.Logger);
 		}
 
 		private static List<string> AcceptThumbprints = new List<string>();
diff --git a/Engine/Platforms/Hololens/Source/Programs/UnrealBuildTool/HoloLensDeploy.cs b/Engine/Platforms/Hololens/Source/Programs/UnrealBuildTool/HoloLensDeploy.cs
--- a/Engine/Platforms/Hololens/Source/Programs/UnrealBuildTool/HoloLensDeploy.cs
+++ b/Engine/Platforms/Hololens/Source/Programs/UnrealBuildTool/HoloLensDeploy.cs
@@ -9,6 +9,7 @@
 using EpicGames.Core;
 using UnrealBuildBase;
 using System.Runtime.Versioning;
+using Microsoft.Extensions.Logging;
 
 namespace UnrealBuildTool
 {
@@ -22,12 +23,17 @@
 		private FileReference? SignToolPath;
 		private string Extension = ".appx";
 
+		public HoloLensDeploy(ILogger InLogger)
+			: base(InLogger)
+		{
+		}
+
 		/// <summary>
 		/// Utility function to delete a file
 		/// </summary>
 		void DeployHelper_DeleteFile(string InFileToDelete)
 		{
-			Log.TraceInformation("HoloLensDeploy.DeployHelper_DeleteFile({0})", InFileToDelete);
+			Logger.LogInformation("HoloLensDeploy.DeployHelper_DeleteFile({File})", InFileToDelete);
 			if (File.Exists(InFileToDelete) == true)
 			{
 				FileAttributes attributes = File.GetAttributes(InFileToDelete);
@@ -46,10 +52,10 @@
 		bool CopySourceToDestDir(string InSourceDirectory, string InDestinationDirectory, string InWildCard,
 			bool bInIncludeSubDirectories, bool bInRemoveDestinationOrphans)
 		{
-			Log.TraceInformation("HoloLensDeploy.CopySourceToDestDir({0}, {1}, {2},...)", InSourceDirectory, InDestinationDirectory, InWildCard);
+			Logger.LogInformation("HoloLensDeploy.CopySourceToDestDir({SourceDir}, {DestDir}, {Wildcard},...)", InSourceDirectory, InDestinationDirectory, InWildCard);
 			if (Directory.Exists(InSourceDirectory) == false)
 			{
-				Log.TraceInformation("Warning: CopySourceToDestDir - SourceDirectory does not exist: {0}", InSourceDirectory);
+				Logger.LogInformation("Warning: CopySourceToDestDir - SourceDirectory does not exist: {SourceDir}", InSourceDirectory);
 				return false;
 			}
 
@@ -91,7 +97,7 @@
 					}
 					catch (Exception exceptionMessage)
 					{
-						Log.TraceInformation("Failed to copy {0} to deployment: {1}", SourceFile, exceptionMessage);
+						Logger.LogInformation("Failed to copy {Source} to deployment: {Ex}", SourceFile, exceptionMessage);
 					}
 				}
 			}
@@ -105,7 +111,7 @@
 					string DestFilename = DestFile.Replace(InDestinationDirectory, "");
 					if (FilesInSource.Contains(DestFilename.ToUpperInvariant()) == false)
 					{
-						Log.TraceInformation("Destination file does not exist in Source - DELETING: {0}", DestFile);
+						Logger.LogInformation("Destination file does not exist in Source - DELETING: {File}", DestFile);
 						//FileAttributes attributes = File.GetAttributes(DestFile);
 						try
 						{
@@ -113,7 +119,7 @@
 						}
 						catch (Exception exceptionMessage)
 						{
-							Log.TraceInformation("Failed to delete {0} from deployment: {1}", DestFile, exceptionMessage);
+							Logger.LogInformation("Failed to delete {DestFile} from deployment: {Ex}", DestFile, exceptionMessage);
 						}
 					}
 				}
@@ -142,13 +148,13 @@
 						DeployHelper_DeleteFile(InDest);
 					}
 				}
-				Log.TraceInformation("HoloLensDeploy.CopyFile({0}, {1}, {2})", InSource, InDest, bForce);
+				Logger.LogInformation("HoloLensDeploy.CopyFile({Source}, {Dest}, {Force})", InSource, InDest, bForce);
 				File.Copy(InSource, InDest, true);
 				File.SetAttributes(InDest, File.GetAttributes(InDest) & ~FileAttributes.ReadOnly);
 			}
 			else
 			{
-				Log.TraceInformation("HoloLensDeploy: File didn't exist - {0}", InSource);
+				Logger.LogInformation("HoloLensDeploy: File didn't exist - {Source}", InSource);
 			}
 		}
 
@@ -317,7 +323,7 @@
 					}
 					else
 					{
-						Log.TraceError("Wrong path to resource \'{0}\', the resource should be in \'{1}\'", ResourceFile.FullName, ResourceFolder.FullName);
+						Logger.LogError("Wrong path to resource \'{Resource}\', the resource should be in \'{Folder}\'", ResourceFile.FullName, ResourceFolder.FullName);
 						throw new BuildException("Failed to generate AppX file.  See log for details.");
 					}
 				}
@@ -420,7 +426,7 @@
 			}
 			catch (Exception exceptionMessage)
 			{
-				Log.TraceError("Failed to delete {0} from deployment: {1}", OutputAppX, exceptionMessage);
+				Logger.LogError("Failed to delete {Output} from deployment: {Ex}", OutputAppX, exceptionMessage);
 				throw new BuildException("Failed to generate AppX file.  See log for details.");
 			}
 
@@ -437,7 +443,7 @@
 			var StartInfo = new ProcessStartInfo(MakeAppXPath!.FullName, MakeAppXCommandLine);
 			StartInfo.UseShellExecute = false;
 			StartInfo.CreateNoWindow = true;
-			var ExitCode = Utils.RunLocalProcessAndPrintfOutput(StartInfo);
+			var ExitCode = Utils.RunLocalProcessAndPrintfOutput(StartInfo, Logger);
 			if (ExitCode < 0)
 			{
 				throw new BuildException("Failed to generate AppX file.  See log for details.");
@@ -449,14 +455,14 @@
 				StartInfo = new ProcessStartInfo(SignToolPath!.FullName, SignToolCommandLine);
 				StartInfo.UseShellExecute = false;
 				StartInfo.CreateNoWindow = true;
-				ExitCode = Utils.RunLocalProcessAndPrintfOutput(StartInfo);
+				ExitCode = Utils.RunLocalProcessAndPrintfOutput(StartInfo, Logger);
 				if (ExitCode < 0)
 				{
 					throw new BuildException("Failed to generate AppX file.  See log for details.");
 				}
 			}
 
-			Log.TraceInformation("AppX successfully packaged to \'{0}{1}\'", OutputName, Extension);
+			Logger.LogInformation("AppX successfully packaged to \'{Output}{Ext}\'", OutputName, Extension);
 			NewReceipt.BuildProducts.Add(new BuildProduct(new FileReference(OutputAppX), BuildProductType.Package));
 		}
 
@@ -478,7 +484,7 @@
 		{
 			// Use the project name if possible - InTarget.AppName changes for 'Client'/'Server' builds
 			string ProjectName = Receipt.ProjectFile != null ? Receipt.ProjectFile.GetFileNameWithoutAnyExtensions() : Receipt.Launch!.GetFileNameWithoutExtension();
-			Log.TraceInformation("Prepping {0} for deployment to {1}", ProjectName, Receipt.Platform.ToString());
+			Logger.LogInformation("Prepping {Project} for deployment to {Platform}", ProjectName, Receipt.Platform.ToString());
 			System.DateTime PrepDeployStartTime = DateTime.UtcNow;
 
 			// Note: TargetReceipt.Read now expands path variables internally.
@@ -495,7 +501,7 @@
 			{
 				SDK = Results.First().Value;
 			}
-			HoloLensExports.InitWindowsSdkToolPath(SDK);
+			HoloLensExports.InitWindowsSdkToolPath(SDK, Logger);
 
 			AddWinMDReferencesFromReceipt(Receipt, Receipt.ProjectDir != null ? Receipt.ProjectDir : Unreal.EngineDirectory, Unreal.EngineDirectory.ParentDirectory!.FullName, SDK);
 
@@ -514,7 +520,7 @@
 			string AbsoluteExeDirectory = Path.GetDirectoryName(ExePaths[0])!;
 			UnrealTargetPlatform Platform = UnrealTargetPlatform.HoloLens;
 			string IntermediateDirectory = Path.Combine(Receipt.ProjectDir != null ? Receipt.ProjectDir.FullName : Unreal.EngineDirectory.FullName, "Intermediate", "Deploy", WindowsExports.GetArchitectureSubpath(Arch));
-			List<string> UpdatedFiles = new HoloLensManifestGenerator().CreateManifest(Platform, Arch, AbsoluteExeDirectory, IntermediateDirectory, Receipt.ProjectFile, Receipt.ProjectDir != null ? Receipt.ProjectDir.FullName : Unreal.EngineDirectory.FullName, TargetConfigs, ExePaths, WinMDReferences)!;
+			List<string> UpdatedFiles = new HoloLensManifestGenerator(Logger).CreateManifest(Platform, Arch, AbsoluteExeDirectory, IntermediateDirectory, Receipt.ProjectFile, Receipt.ProjectDir != null ? Receipt.ProjectDir.FullName : Unreal.EngineDirectory.FullName, TargetConfigs, ExePaths, WinMDReferences)!;
 
 			PrepForUATPackageOrDeploy(Receipt.ProjectFile, ProjectName, Receipt.ProjectDir != null ? Receipt.ProjectDir.FullName : Unreal.EngineDirectory.FullName, Arch, TargetConfigs, ExePaths, RelativeEnginePath, false, "", false);
 			MakePackage(Receipt, NewReceipt, Arch, UpdatedFiles);
@@ -524,7 +530,7 @@
 
 			// Log out the time taken to deploy...
 			double PrepDeployDuration = (DateTime.UtcNow - PrepDeployStartTime).TotalSeconds;
-			Log.TraceInformation("HoloLens deployment preparation took {0:0.00} seconds", PrepDeployDuration);
+			Logger.LogInformation("HoloLens deployment preparation took {Time:0.00} seconds", PrepDeployDuration);
 
 			return true;
 		}
@@ -544,7 +550,7 @@
 						DestPath = Dep.Path.FullName.Replace(Unreal.EngineDirectory.FullName, Path.Combine(DestRelativeTo, "Engine"));
 						DestPath = DestPath.Replace(SourceProjectDir.FullName, Path.Combine(DestRelativeTo, SourceProjectDir.GetDirectoryName()));
 						DestPath = Utils.MakePathRelativeTo(DestPath, DestRelativeTo);
-						WinMDReferences.Add(new WinMDRegistrationInfo(new FileReference(WinMDFile), DestPath, SDKVersion));
+						WinMDReferences.Add(new WinMDRegistrationInfo(new FileReference(WinMDFile), DestPath, SDKVersion, Logger));
 					}
 				}
 			}
diff --git a/Engine/Platforms/Hololens/Source/Programs/UnrealBuildTool/HoloLensExports.cs b/Engine/Platforms/Hololens/Source/Programs/UnrealBuildTool/HoloLensExports.cs
--- a/Engine/Platforms/Hololens/Source/Programs/UnrealBuildTool/HoloLensExports.cs
+++ b/Engine/Platforms/Hololens/Source/Programs/UnrealBuildTool/HoloLensExports.cs
@@ -7,6 +7,7 @@
 using System.Text;
 using System.Threading.Tasks;
 using EpicGames.Core;
+using Microsoft.Extensions.Logging;
 
 namespace UnrealBuildTool
 {
@@ -21,9 +22,9 @@
 		/// <summary>
 		/// 
 		/// </summary>
-		public HoloLensExports()
+		public HoloLensExports(ILogger Logger)
 		{
-			InnerDeploy = new HoloLensDeploy();
+			InnerDeploy = new HoloLensDeploy(Logger);
 		}
 
 		/// <summary>
@@ -71,10 +72,11 @@
 		/// 
 		/// </summary>
 		/// <param name="SdkVersion"></param>
+		/// <param name="Logger"></param>
 		/// <returns></returns>
-		public static bool InitWindowsSdkToolPath(string SdkVersion)
+		public static bool InitWindowsSdkToolPath(string SdkVersion, ILogger Logger)
 		{
-			return HoloLensToolChain.InitWindowsSdkToolPath(SdkVersion);
+			return HoloLensToolChain.InitWindowsSdkToolPath(SdkVersion, Logger);
 		}
 
 		/// <summary>
@@ -82,10 +84,11 @@
 		/// </summary>
 		/// <param name="DLCFile"></param>
 		/// <param name="OutputDirectory"></param>
-		public static void CreateManifestForDLC(FileReference DLCFile, DirectoryReference OutputDirectory)
+		/// <param name="Logger"></param>
+		public static void CreateManifestForDLC(FileReference DLCFile, DirectoryReference OutputDirectory, ILogger Logger)
 		{
 			string IntermediateDirectory = DirectoryReference.Combine(DLCFile.Directory, "Intermediate", "Deploy").FullName;
-			new HoloLensManifestGenerator().CreateManifest(UnrealTargetPlatform.HoloLens, WindowsArchitecture.ARM64, OutputDirectory.FullName, IntermediateDirectory, DLCFile, DLCFile.Directory.FullName, new List<UnrealTargetConfiguration>(), new List<string>(), null);
+			new HoloLensManifestGenerator(Logger).CreateManifest(UnrealTargetPlatform.HoloLens, WindowsArchitecture.ARM64, OutputDirectory.FullName, IntermediateDirectory, DLCFile, DLCFile.Directory.FullName, new List<UnrealTargetConfiguration>(), new List<string>(), null);
 		}
 
 		/// <summary>
diff --git a/Engine/Platforms/Hololens/Source/Programs/UnrealBuildTool/HoloLensProjectGenerator.cs b/Engine/Platforms/Hololens/Source/Programs/UnrealBuildTool/HoloLensProjectGenerator.cs
--- a/Engine/Platforms/Hololens/Source/Programs/UnrealBuildTool/HoloLensProjectGenerator.cs
+++ b/Engine/Platforms/Hololens/Source/Programs/UnrealBuildTool/HoloLensProjectGenerator.cs
@@ -10,6 +10,7 @@
 using EpicGames.Core;
 using UnrealBuildBase;
 using System.Runtime.Versioning;
+using Microsoft.Extensions.Logging;
 
 namespace UnrealBuildTool
 {
@@ -25,8 +26,9 @@
 		/// Constructor
 		/// </summary>
 		/// <param name="Arguments">Command line arguments passed to the project generator</param>
-		public HoloLensProjectGenerator(CommandLineArguments Arguments)
-			: base(Arguments)
+		/// <param name="Logger"></param>
+		public HoloLensProjectGenerator(CommandLineArguments Arguments, ILogger Logger)
+			: base(Arguments, Logger)
 		{
 		}
 
@@ -81,14 +83,14 @@
 				VCProjectFileContent.Append("		<WindowsTargetPlatformVersion>" + MaxTestedVersion + "</WindowsTargetPlatformVersion>" + ProjectFileGenerator.NewLine);
 			}
 
-			WindowsCompiler Compiler = WindowsPlatform.GetDefaultCompiler(TargetFilePath, WindowsArchitecture.x64);
-			DirectoryReference? PlatformWinMDLocation = HoloLensPlatform.GetCppCXMetadataLocation(Compiler, "Latest", WindowsArchitecture.x64);
+			WindowsCompiler Compiler = WindowsPlatform.GetDefaultCompiler(TargetFilePath, WindowsArchitecture.x64, Logger);
+			DirectoryReference? PlatformWinMDLocation = HoloLensPlatform.GetCppCXMetadataLocation(Compiler, "Latest", WindowsArchitecture.x64, Logger);
 			if (PlatformWinMDLocation == null || !FileReference.Exists(FileReference.Combine(PlatformWinMDLocation, "platform.winmd")))
 			{
 				throw new BuildException("Unable to find platform.winmd for {0} toolchain; '{1}' is an invalid version", WindowsPlatform.GetCompilerName(Compiler), "Latest");
 			}
-			string FoundationWinMDPath = HoloLensPlatform.GetLatestMetadataPathForApiContract("Windows.Foundation.FoundationContract", Compiler);
-			string UniversalWinMDPath = HoloLensPlatform.GetLatestMetadataPathForApiContract("Windows.Foundation.UniversalApiContract", Compiler);
+			string FoundationWinMDPath = HoloLensPlatform.GetLatestMetadataPathForApiContract("Windows.Foundation.FoundationContract", Compiler, Logger);
+			string UniversalWinMDPath = HoloLensPlatform.GetLatestMetadataPathForApiContract("Windows.Foundation.UniversalApiContract", Compiler, Logger);
 			VCProjectFileContent.Append("		<AdditionalOptions>/ZW /ZW:nostdlib</AdditionalOptions>" + ProjectFileGenerator.NewLine);
 			VCProjectFileContent.Append("		<NMakePreprocessorDefinitions>$(NMakePreprocessorDefinitions);PLATFORM_HOLOLENS=1;HOLOLENS=1;</NMakePreprocessorDefinitions>" + ProjectFileGenerator.NewLine);
 			if (PlatformWinMDLocation != null)
@@ -110,7 +112,7 @@
 
 			DirectoryReference? folder;
 			VersionNumber? version;
-			if(WindowsPlatform.TryGetWindowsSdkDir("Latest", out version, out folder))
+			if(WindowsPlatform.TryGetWindowsSdkDir("Latest", Logger, out version, out folder))
 			{
 				//SDKFolder = folder.FullName;
 				SDKVersion = version.ToString();
diff --git a/Engine/Platforms/Hololens/Source/Programs/UnrealBuildTool/HoloLensToolChain.cs b/Engine/Platforms/Hololens/Source/Programs/UnrealBuildTool/HoloLensToolChain.cs
--- a/Engine/Platforms/Hololens/Source/Programs/UnrealBuildTool/HoloLensToolChain.cs
+++ b/Engine/Platforms/Hololens/Source/Programs/UnrealBuildTool/HoloLensToolChain.cs
@@ -7,6 +7,7 @@
 using System.Linq;
 using System.Text;
 using System.Text.RegularExpressions;
+using Microsoft.Extensions.Logging;
 using Microsoft.Win32;
 using EpicGames.Core;
 using UnrealBuildBase;
@@ -25,12 +26,12 @@
 	[SupportedOSPlatform("windows")]
 	class HoloLensToolChain : VCToolChain
 	{
-		public HoloLensToolChain(ReadOnlyTargetRules Target)
-			: base(Target)
+		public HoloLensToolChain(ReadOnlyTargetRules Target, ILogger Logger)
+			: base(Target, Logger)
 		{
 			// by default tools chains don't parse arguments, but we want to be able to check the -architectures flag defined above. This is
 			// only necessary when AndroidToolChain is used during UAT
-			CommandLine.ParseArguments(Environment.GetCommandLineArgs(), this);
+			CommandLine.ParseArguments(Environment.GetCommandLineArgs(), this, Logger);
 		}
 
 		protected override void AppendCLArguments_Global(CppCompileEnvironment CompileEnvironment, List<string> Arguments)
@@ -148,20 +149,20 @@
 		private static DirectoryReference? CurrentWindowsSdkBinDir = null;
 		private static Version? CurrentWindowsSdkVersion;
 
-		public static bool InitWindowsSdkToolPath(string SdkVersion)
+		public static bool InitWindowsSdkToolPath(string SdkVersion, ILogger Logger)
 		{
 			if (string.IsNullOrEmpty(SdkVersion))
 			{
-				Log.TraceError("WinSDK version is empty");
+				Logger.LogError("WinSDK version is empty");
 				return false;
 			}
 
 			VersionNumber OutSdkVersion;
 			DirectoryReference OutSdkDir;
 
-			if (!WindowsPlatform.TryGetWindowsSdkDir(SdkVersion, out OutSdkVersion!, out OutSdkDir!))
+			if (!WindowsPlatform.TryGetWindowsSdkDir(SdkVersion, Logger, out OutSdkVersion!, out OutSdkDir!))
 			{
-				Log.TraceError("Failed to find WinSDK " + SdkVersion);
+				Logger.LogError("Failed to find WinSDK {SdkVersion}", SdkVersion);
 				return false;
 			}
 
@@ -169,7 +170,7 @@
 
 			if(!DirectoryReference.Exists(WindowsSdkBinDir))
 			{
-				Log.TraceError("WinSDK " + SdkVersion + " doesn't exit");
+				Logger.LogError("WinSDK {SdkVersion} doesn't exit", SdkVersion);
 				return false;
 			}
 
diff --git a/Engine/Platforms/Hololens/Source/Programs/UnrealBuildTool/ManifestGenerator.cs b/Engine/Platforms/Hololens/Source/Programs/UnrealBuildTool/ManifestGenerator.cs
--- a/Engine/Platforms/Hololens/Source/Programs/UnrealBuildTool/ManifestGenerator.cs
+++ b/Engine/Platforms/Hololens/Source/Programs/UnrealBuildTool/ManifestGenerator.cs
@@ -11,6 +11,7 @@
 using EpicGames.Core;
 using UnrealBuildBase;
 using System.Runtime.Versioning;
+using Microsoft.Extensions.Logging;
 
 namespace UnrealBuildTool
 {
@@ -54,8 +55,21 @@
 		private bool IsDlc;
 		private Dictionary<string, string>? ParsedDlcInfo;
 
+		/// <summary>
+		/// Logger for output
+		/// </summary>
+		protected readonly ILogger Logger;
 
 		/// <summary>
+		/// Constructor
+		/// </summary>
+		/// <param name="InLogger"></param>
+		public HoloLensManifestGenerator(ILogger InLogger)
+		{
+			Logger = InLogger;
+		}
+
+		/// <summary>
 		/// Retrieve a package configuration option from the deprecated [AppxManifest] INI settings and return the value.
 		/// NOTE: Do not use this method for storing values, it is left in purely for compatibility. All package configuration
 		/// should be done using the target settings in the editor or in the [/Script/HoloLensPlatformEditor.HoloLensTargetSettings] INI section.
@@ -101,13 +115,13 @@
 				int LenOfSetting = BaseSetting.Substring(NextSetting + 1).IndexOfAny(VariableMarkers);
 				if (LenOfSetting < 0)
 				{
-					Log.TraceError("Could not parse setting {0}. Unmatched variable symbol '{1}'", LookupString, BaseSetting[NextSetting]);
+					Logger.LogError("Could not parse setting {Str}. Unmatched variable symbol '{Var}'", LookupString, BaseSetting[NextSetting]);
 					return InterprettedSetting + BaseSetting;
 				}
 				if (BaseSetting[NextSetting] != BaseSetting[NextSetting + LenOfSetting + 1])
 				{
 					// Probable nested operators
-					Log.TraceError("Could not parse setting {0}. Mismatched variable symbols '{1}' and '{2}'", LookupString, BaseSetting[NextSetting], BaseSetting[NextSetting + LenOfSetting + 1]);
+					Logger.LogError("Could not parse setting {Str}. Mismatched variable symbols '{Var1}' and '{Var2}'", LookupString, BaseSetting[NextSetting], BaseSetting[NextSetting + LenOfSetting + 1]);
 					return InterprettedSetting + BaseSetting;
 				}
 
@@ -147,7 +161,7 @@
 							string InsertSource = Path.Combine(ProjectPath!, VariableName.Substring(VariableName.IndexOf(':') + 1));
 							if (!File.Exists(InsertSource))
 							{
-								Log.TraceWarning("Invalid path for insertion: {0}", InsertSource);
+								Logger.LogWarning("Invalid path for insertion: {Insert}", InsertSource);
 								// @todo: Can't think of a way to insert valid XML in this case, so it's just going to be left out.
 								// It would be better at least to insert something that would lead back here (as is done with
 								// "InvalidIniValue" above.
@@ -160,7 +174,7 @@
 							}
 							catch (Exception)
 							{
-								Log.TraceWarning("Error while trying to read data for insert from {0}.", InsertSource);
+								Logger.LogWarning("Error while trying to read data for insert from {Insert}.", InsertSource);
 								// @todo: Can't think of a way to insert valid XML in this case, so it's just going to be left out.
 								// It would be better at least to insert something that would lead back here (as is done with
 								// "InvalidIniValue" above.
@@ -316,7 +330,7 @@
 						}
 						else
 						{
-							Log.TraceWarning("Unable to parse AppxManifest variable value for {0}.", VariableName);
+							Logger.LogWarning("Unable to parse AppxManifest variable value for {VariableName}.", VariableName);
 							// @todo: Is there any better way to handle not finding the value? If we leave a value blank it will
 							// likely produce invalid XML and be difficult to trace. At least this hardcoded string should lead
 							// users back here.
@@ -350,12 +364,12 @@
 				}
 				catch (Exception)
 				{
-					Log.TraceError("Could not create directory {0}.", TargetDirectory);
+					Logger.LogError("Could not create directory {TargetDir}.", TargetDirectory);
 					return false;
 				}
 				if (!Directory.Exists(TargetDirectory))
 				{
-					Log.TraceError("Path {0} does not exist or is not a directory.", TargetDirectory);
+					Logger.LogError("Path {TargetDir} does not exist or is not a directory.", TargetDirectory);
 					return false;
 				}
 			}
@@ -370,7 +384,7 @@
 		{
 			if (!File.Exists(IntermediatePath))
 			{
-				Log.TraceError("Tried to copy non-existant intermediate file {0}.", IntermediatePath);
+				Logger.LogError("Tried to copy non-existant intermediate file {IntermediatePath}.", IntermediatePath);
 				return;
 			}
 
@@ -395,7 +409,7 @@
 					}
 					catch (Exception)
 					{
-						Log.TraceError("Could not replace file {0}.", TargetPath);
+						Logger.LogError("Could not replace file {TargetPath}.", TargetPath);
 						return;
 					}
 				}
@@ -410,7 +424,7 @@
 				}
 				catch (Exception)
 				{
-					Log.TraceError("Unable to copy file {0}.", TargetPath);
+					Logger.LogError("Unable to copy file {TargetPath}.", TargetPath);
 					return;
 				}
 				UpdatedFilePaths!.Add(TargetPath);
@@ -474,7 +488,7 @@
 				string TargetResourcePath = Path.Combine(TargetPath, SourceResourceFile.Substring(SourcePath.Length + 1));
 				if (!CreateCheckDirectory(Path.GetDirectoryName(TargetResourcePath)!))
 				{
-					Log.TraceError("Unable to create intermediate directory {0}.", Path.GetDirectoryName(TargetResourcePath));
+					Logger.LogError("Unable to create intermediate directory {IntermediateDir}.", Path.GetDirectoryName(TargetResourcePath));
 					continue;
 				}
 				if (!File.Exists(TargetResourcePath))
@@ -485,7 +499,7 @@
 					}
 					catch (Exception)
 					{
-						Log.TraceError("Unable to copy file {0} to {1}.", SourceResourceFile, TargetResourcePath);
+						Logger.LogError("Unable to copy file {SourceFile} to {TargetFile}.", SourceResourceFile, TargetResourcePath);
 						return false;
 					}
 				}
@@ -503,7 +517,7 @@
 				string TargetResourcePath = Path.Combine(TargetPath, SourceResourceFile.Substring(SourcePath.Length + 1));
 				if (!CreateCheckDirectory(Path.GetDirectoryName(TargetResourcePath)!))
 				{
-					Log.TraceError("Unable to create intermediate directory {0}.", Path.GetDirectoryName(TargetResourcePath));
+					Logger.LogError("Unable to create intermediate directory {IntDir}.", Path.GetDirectoryName(TargetResourcePath));
 					continue;
 				}
 				if (!File.Exists(TargetResourcePath))
@@ -514,7 +528,7 @@
 					}
 					catch (Exception)
 					{
-						Log.TraceError("Unable to copy file {0} to {1}.", SourceResourceFile, TargetResourcePath);
+						Logger.LogError("Unable to copy file {SourceFile} to {TargetFile}.", SourceResourceFile, TargetResourcePath);
 						throw;
 					}
 				}
@@ -570,7 +584,7 @@
 						}
 						catch (Exception)
 						{
-							Log.TraceError("Could not remove stale resource file {0}.", TargetResourceFile);
+							Logger.LogError("Could not remove stale resource file {TargetFile}.", TargetResourceFile);
 						}
 					}
 				}
@@ -616,14 +630,14 @@
 						}
 						catch (Exception)
 						{
-							Log.TraceWarning("Could not remove file {0} to remove directory {1}.", FileToRemove, InDirectoryToDelete);
+							Logger.LogWarning("Could not remove file {File} to remove directory {Dir}.", FileToRemove, InDirectoryToDelete);
 						}
 					}
 					Directory.Delete(InDirectoryToDelete, true);
 				}
 				catch (Exception)
 				{
-					Log.TraceWarning("Could not remove directory {0}.", InDirectoryToDelete);
+					Logger.LogWarning("Could not remove directory {Dir}.", InDirectoryToDelete);
 				}
 			}
 
@@ -677,7 +691,7 @@
                     continue;
                 }
 
-                Log.TraceLog("Looking at file: " + IniFileName.FullName);
+                Logger.LogDebug("Looking at file: {File}", IniFileName.FullName);
                 projectIniFilePaths.Add(IniFileName.FullName);
 
                 bool isInDesiredSection = false;
@@ -698,7 +712,7 @@
 
 					if (line.StartsWith(keyName.Trim() + "="))
                     {
-                        Log.TraceLog("Found string match.");
+                        Logger.LogDebug("Found string match.");
                         iniUpdated = true;
                         rewriteCurrentFile = true;
 
@@ -747,17 +761,17 @@
 			// Check parameter values are valid
 			if (InTargetConfigs.Count != InExecutables.Count)
 			{
-				Log.TraceError("The number of target configurations ({0}) and executables ({1}) passed to manifest generation do not match.", InTargetConfigs.Count, InExecutables.Count);
+				Logger.LogError("The number of target configurations ({ConfigCount}) and executables ({ExeCount}) passed to manifest generation do not match.", InTargetConfigs.Count, InExecutables.Count);
 				return null;
 			}
 			if (File.Exists(InOutputPath))
 			{
-				Log.TraceWarning("InOutputPath {0} is a file. Should be a directory. Continuing using parent directory.", InOutputPath);
+				Logger.LogWarning("InOutputPath {InOutputPath} is a file. Should be a directory. Continuing using parent directory.", InOutputPath);
 				InOutputPath = Path.GetDirectoryName(InOutputPath)!;
 			}
 			if (File.Exists(InIntermediatePath))
 			{
-				Log.TraceWarning("InIntermediatePath {0} is a file. Should be a directory. Continuing using parent directory.", InIntermediatePath);
+				Logger.LogWarning("InIntermediatePath {IntPath} is a file. Should be a directory. Continuing using parent directory.", InIntermediatePath);
 				InIntermediatePath = Path.GetDirectoryName(InIntermediatePath)!;
 			}
 			if (!CreateCheckDirectory(InOutputPath))
@@ -803,7 +817,7 @@
 				}
 				catch (Exception)
 				{
-					Log.TraceError("Could not create directory {0}.", IntermediateResourceDirectory);
+					Logger.LogError("Could not create directory {IntDir}.", IntermediateResourceDirectory);
 					return null;
 				}
 			}
@@ -841,7 +855,7 @@
 
 					if (ParsedDlcInfo == null)
 					{
-						Log.TraceWarning("Could not map {0} to a Store identity.  Using a temporary identity to enable local deployment.  For Store upload configure identity in the HoloLens Project Settings.", InProjectFile);
+						Logger.LogWarning("Could not map {ProjectFile} to a Store identity.  Using a temporary identity to enable local deployment.  For Store upload configure identity in the HoloLens Project Settings.", InProjectFile);
 						ParsedDlcInfo = new Dictionary<string, string>();
 						ParsedDlcInfo["PluginName"] = InProjectFile.GetFileNameWithoutExtension();
 						ParsedDlcInfo["PackageIdentityName"] = ParsedDlcInfo["PluginName"];
@@ -878,7 +892,7 @@
 			GameIni.GetArray("/Script/UnrealEd.ProjectPackagingSettings", "CulturesToStage", out CulturesToStageWithDuplicates);
 			if (CulturesToStageWithDuplicates == null || CulturesToStageWithDuplicates.Count < 1)
 			{
-				Log.TraceError("At least one culture must be selected to stage.");
+				Logger.LogError("At least one culture must be selected to stage.");
 				return null;
 			}
 
@@ -896,7 +910,7 @@
 				string IntermediateStringResourceFile = Path.Combine(IntermediateStringResourcePath, "resources.resw");
 				if (!CreateCheckDirectory(IntermediateStringResourcePath))
 				{
-					Log.TraceWarning("Failed to create {0}.  Culture {1} resources not staged.", IntermediateStringResourcePath, Culture);
+					Logger.LogWarning("Failed to create {Path}.  Culture {Culture} resources not staged.", IntermediateStringResourcePath, Culture);
 					CulturesToStage.RemoveAt(i);
 					--i;
 					continue;
@@ -906,7 +920,7 @@
 
 			if (CulturesToStage.Count == 0)
 			{
-				Log.TraceError("Failed to create intermediate files for any culture.  Manifest could not be generated.");
+				Logger.LogError("Failed to create intermediate files for any culture.  Manifest could not be generated.");
 				return null;
 			}
 
@@ -917,7 +931,7 @@
                 string currentVersion;
                 if (GameIni.GetString("/Script/EngineSettings.GeneralProjectSettings", "ProjectVersion", out currentVersion))
                 {
-                    Log.TraceLog("Automatically incrementing version. Starting version is " + currentVersion);
+                    Logger.LogDebug("Automatically incrementing version. Starting version is {Version}", currentVersion);
 
                     string[] versionEntries = currentVersion.Split(new char[] { '.' });
                     if (versionEntries.Length == 4)
@@ -926,22 +940,22 @@
                         versionEntryToIncrement++;
 
                         string newVersion = string.Concat(versionEntries[0], ".", versionEntries[1], ".", versionEntryToIncrement, ".", versionEntries[3]);
-                        Log.TraceLog("Writing new version string: " + newVersion);
+                        Logger.LogDebug("Writing new version string: {Version}", newVersion);
 
                         // Update GameIni with newVersion.
                         if (!UpdateProjectIniString(InProjectFile, TargetPlatform, "/Script/EngineSettings.GeneralProjectSettings", "ProjectVersion", newVersion))
                         {
-                            Log.TraceWarning("Auto incrementing the project version was unsuccessful.");
+                            Logger.LogWarning("Auto incrementing the project version was unsuccessful.");
                         }
                     }
                     else
                     {
-                        Log.TraceWarning("Auto increment was desired, but the number of decimals in the version string was unexpected.");
+                        Logger.LogWarning("Auto increment was desired, but the number of decimals in the version string was unexpected.");
                     }
                 }
                 else
                 {
-                    Log.TraceWarning("Auto increment was desired, but the existing version could not be identified.");
+                    Logger.LogWarning("Auto increment was desired, but the existing version could not be identified.");
                 }
             }
 
@@ -1023,7 +1037,7 @@
 				StartInfo.UseShellExecute = false;
 				StartInfo.RedirectStandardOutput = true;
 				StartInfo.CreateNoWindow = true;
-				int ExitCode = Utils.RunLocalProcessAndLogOutput(StartInfo);
+				int ExitCode = Utils.RunLocalProcessAndLogOutput(StartInfo, Logger);
 				if (ExitCode < 0)
 				{
 					throw new BuildException("Failed to generate config file for Package Resource Index.  See log for details.");
@@ -1102,7 +1116,7 @@
 					}
 					catch (Exception)
 					{
-						Log.TraceError("Could not delete file {0}.", OldPri);
+						Logger.LogError("Could not delete file {File}.", OldPri);
 					}
 				}
 
@@ -1116,7 +1130,7 @@
 				StartInfo.CreateNoWindow = true;
 				StartInfo.StandardErrorEncoding = System.Text.Encoding.Unicode;
 				StartInfo.StandardOutputEncoding = System.Text.Encoding.Unicode;
-				ExitCode = Utils.RunLocalProcessAndLogOutput(StartInfo);
+				ExitCode = Utils.RunLocalProcessAndLogOutput(StartInfo, Logger);
 				if (ExitCode < 0)
 				{
 					throw new BuildException("Failed to generate Package Resource Index file.  See log for details.");
@@ -1144,12 +1158,12 @@
 		{
 			if (File.Exists(InOutputPath))
 			{
-				Log.TraceWarning("InOutputPath {0} is a file. Should be a directory. Continuing using parent directory.", InOutputPath);
+				Logger.LogWarning("InOutputPath {InOutputPath} is a file. Should be a directory. Continuing using parent directory.", InOutputPath);
 				InOutputPath = Path.GetDirectoryName(InOutputPath)!;
 			}
 			if (File.Exists(InIntermediatePath))
 			{
-				Log.TraceWarning("InIntermediatePath {0} is a file. Should be a directory. Continuing using parent directory.", InIntermediatePath);
+				Logger.LogWarning("InIntermediatePath {InIntermediatePath} is a file. Should be a directory. Continuing using parent directory.", InIntermediatePath);
 				InIntermediatePath = Path.GetDirectoryName(InIntermediatePath)!;
 			}
 			if (!CreateCheckDirectory(InOutputPath))
@@ -1183,7 +1197,7 @@
 				}
 				catch (Exception)
 				{
-					Log.TraceError("Could not create directory {0}.", IntermediateResourceDirectory);
+					Logger.LogError("Could not create directory {IntDir}.", IntermediateResourceDirectory);
 					return null;
 				}
 			}
@@ -1221,7 +1235,7 @@
 
 					if (ParsedDlcInfo == null)
 					{
-						Log.TraceWarning("Could not map {0} to a Store identity.  Using a temporary identity to enable local deployment.  For Store upload configure identity in the HoloLens Project Settings.", InProjectFile);
+						Logger.LogWarning("Could not map {ProjectFile} to a Store identity.  Using a temporary identity to enable local deployment.  For Store upload configure identity in the HoloLens Project Settings.", InProjectFile);
 						ParsedDlcInfo = new Dictionary<string, string>();
 						ParsedDlcInfo["PluginName"] = InProjectFile.GetFileNameWithoutExtension();
 						ParsedDlcInfo["PackageIdentityName"] = ParsedDlcInfo["PluginName"];
@@ -1258,7 +1272,7 @@
 			GameIni.GetArray("/Script/UnrealEd.ProjectPackagingSettings", "CulturesToStage", out CulturesToStageWithDuplicates);
 			if (CulturesToStageWithDuplicates == null || CulturesToStageWithDuplicates.Count < 1)
 			{
-				Log.TraceError("At least one culture must be selected to stage.");
+				Logger.LogError("At least one culture must be selected to stage.");
 				return null;
 			}
 
@@ -1276,7 +1290,7 @@
 				string IntermediateStringResourceFile = Path.Combine(IntermediateStringResourcePath, "resources.resw");
 				if (!CreateCheckDirectory(IntermediateStringResourcePath))
 				{
-					Log.TraceWarning("Failed to create {0}.  Culture {1} resources not staged.", IntermediateStringResourcePath, Culture);
+					Logger.LogWarning("Failed to create {Path}.  Culture {Culture} resources not staged.", IntermediateStringResourcePath, Culture);
 					CulturesToStage.RemoveAt(i);
 					--i;
 					continue;
@@ -1286,7 +1300,7 @@
 
 			if (CulturesToStage.Count == 0)
 			{
-				Log.TraceError("Failed to create intermediate files for any culture.  Manifest could not be generated.");
+				Logger.LogError("Failed to create intermediate files for any culture.  Manifest could not be generated.");
 				return null;
 			}
 
@@ -1355,7 +1369,7 @@
 					}
 					else
 					{
-						Log.TraceError("Unable to stage package logo.");
+						Logger.LogError("Unable to stage package logo.");
 					}
 
 					XmlElement ResourcePackageParam = AppxManifestXmlDocument.CreateElement("ResourcePackage");
@@ -1694,7 +1708,7 @@
 				}
 				else
 				{ 
-					Log.TraceVerbose("No localized value for {0} in culture {1}.  Neutral value ({2}) will be used", ResourceEntryName, CulturesToStage[i], NeutralValue);
+					Logger.LogDebug("No localized value for {ResourceName} in culture {Culture}.  Neutral value ({Neutral}) will be used", ResourceEntryName, CulturesToStage[i], NeutralValue);
 				}
 			}
 
@@ -1748,12 +1762,12 @@
 				}
 				else if (bNodeRequired)
 				{
-					Log.TraceError("Node {0} that requires a value is empty.", Child.Name);
+					Logger.LogError("Node {Name} that requires a value is empty.", Child.Name);
 				}
 			}
 			else if (bNodeRequired)
 			{
-				Log.TraceError("Unable to create required manifest entry");
+				Logger.LogError("Unable to create required manifest entry");
 			}
 		}
 
@@ -1831,8 +1845,8 @@
 			string ReturnVal = Regex.Replace(InPackageName, "[^-.A-Za-z0-9]", "");
 			if (ReturnVal == null || ReturnVal.Length <= 0)
 			{
-				Log.TraceError("Invalid package name {0}. Package names must only contain letters, numbers, dash, and period and must be at least one character long.", InPackageName);
-				Log.TraceError("Consider using the setting [/Script/HoloLensPlatformEditor.HoloLensTargetSettings]:PackageName to provide a HoloLens specific value.");
+				Logger.LogError("Invalid package name {PackageName}. Package names must only contain letters, numbers, dash, and period and must be at least one character long.", InPackageName);
+				Logger.LogError("Consider using the setting [/Script/HoloLensPlatformEditor.HoloLensTargetSettings]:PackageName to provide a HoloLens specific value.");
 				ReturnVal = String.Empty;
 			}
 			return ReturnVal;
@@ -1911,7 +1925,7 @@
 			}
 			else
 			{
-				Log.TraceError("Unable to stage package logo.");
+				Logger.LogError("Unable to stage package logo.");
 			}
 
 			return Properties;
@@ -1997,7 +2011,7 @@
 			// Check that we have a valid number of cultures.
 			if (CulturesToStage!.Count < 1 || CulturesToStage.Count >= MaxResourceEntries)
 			{
-				Log.TraceWarning("Incorrect number of cultures to stage. There must be between 1 and {0} cultures selected.", MaxResourceEntries);
+				Logger.LogWarning("Incorrect number of cultures to stage. There must be between 1 and {MaxEntries} cultures selected.", MaxResourceEntries);
 			}
 
 			// Create the culture list.
@@ -2024,12 +2038,12 @@
 
 			if (TargetConfigs.Count < 1)
 			{
-				Log.TraceError("No configurations to deploy");
+				Logger.LogError("No configurations to deploy");
 				return Applications;
 			}
 			if (TargetConfigs.Count != Executables.Count)
 			{
-				Log.TraceError("The number of executables does not match the number of configurations.");
+				Logger.LogError("The number of executables does not match the number of configurations.");
 				return Applications;
 			}
 
@@ -2056,8 +2070,8 @@
 			}
 			if (ReturnVal == null || ReturnVal.Length <= 0)
 			{
-				Log.TraceError("Invalid application ID {0}. Application IDs must only contain letters and numbers. And they must begin with a letter.", InApplicationId);
-				Log.TraceError("Consider using the setting [/Script/HoloLensPlatformEditor.HoloLensTargetSettings]:ValidateApplicationName to provide a HoloLens specific value.");
+				Logger.LogError("Invalid application ID {ApplicationId}. Application IDs must only contain letters and numbers. And they must begin with a letter.", InApplicationId);
+				Logger.LogError("Consider using the setting [/Script/HoloLensPlatformEditor.HoloLensTargetSettings]:ValidateApplicationName to provide a HoloLens specific value.");
 				ReturnVal = String.Empty;
 			}
 			return ReturnVal;
@@ -2137,7 +2151,7 @@
 			}
 			else
 			{
-				Log.TraceError("Unable to stage application logo.");
+				Logger.LogError("Unable to stage application logo.");
 			}
 
 			XmlAttribute SmallLogo = AppxManifestXmlDocument.CreateAttribute("Square44x44Logo");
@@ -2148,7 +2162,7 @@
 			}
 			else
 			{
-				Log.TraceError("Unable to stage application small logo.");
+				Logger.LogError("Unable to stage application small logo.");
 			}
 
 			XmlNode SplashScreen = GetSplashScreen(ApplicationIndex);
@@ -2191,10 +2205,10 @@
 					StartInfo.UseShellExecute = false;
 					StartInfo.RedirectStandardOutput = true;
 					StartInfo.CreateNoWindow = true;
-					int ExitCode = Utils.RunLocalProcessAndPrintfOutput(StartInfo);
+					int ExitCode = Utils.RunLocalProcessAndPrintfOutput(StartInfo, Logger);
 					if (ExitCode < 0)
 					{
-						Log.TraceError("GLTF packaging failed. See log for details.");
+						Logger.LogError("GLTF packaging failed. See log for details.");
 						throw new BuildException("GLTF packaging failed. See log for details.");
 					}
 					File.SetLastWriteTimeUtc(to, File.GetLastWriteTimeUtc(from));
@@ -2209,7 +2223,7 @@
 			}
 			else
 			{
-				Log.TraceError("Unable to stage application wide logo.");
+				Logger.LogError("Unable to stage application wide logo.");
 			}
 
 			bool bUseNameForLogo;
@@ -2255,7 +2269,7 @@
 			}
 			else
 			{
-				Log.TraceError("Unable to stage splash screen image.");
+				Logger.LogError("Unable to stage splash screen image.");
 			}
 
 			return SplashScreen;
@@ -2388,7 +2402,7 @@
 
 			DirectoryReference SDKRootFolder;
 			VersionNumber? version;
-			if (WindowsPlatform.TryGetWindowsSdkDir("Latest", out version, out SDKRootFolder!))
+			if (WindowsPlatform.TryGetWindowsSdkDir("Latest", Logger, out version, out SDKRootFolder!))
 			{
 				SDKVersion = version.ToString();
 			}
@@ -2397,11 +2411,11 @@
 			PhoneSchemaFolder = DirectoryReference.Combine(SDKRootFolder, "Extension SDKs", "WindowsMobile", SDKVersion.ToString(), "Include", "WinRT");
 
 			IEnumerable<DirectoryReference>? VSInstallDirs;
-			if (null != (VSInstallDirs = WindowsPlatform.TryGetVSInstallDirs(WindowsCompiler.VisualStudio2019)))
+			if (null != (VSInstallDirs = WindowsPlatform.TryGetVSInstallDirs(WindowsCompiler.VisualStudio2019, Logger)))
 			{
 				VSSchemaFolder = DirectoryReference.Combine(VSInstallDirs.First(), "Xml", "Schemas");
 			}
-			else if (null != (VSInstallDirs = WindowsPlatform.TryGetVSInstallDirs(WindowsCompiler.VisualStudio2022)))
+			else if (null != (VSInstallDirs = WindowsPlatform.TryGetVSInstallDirs(WindowsCompiler.VisualStudio2022, Logger)))
 			{
 				VSSchemaFolder = DirectoryReference.Combine(VSInstallDirs.First(), "Xml", "Schemas");
 			}
@@ -2462,14 +2476,14 @@
 			}
 			catch (System.Xml.Schema.XmlSchemaException e)
 			{
-				string InvalidSchemaWarning =
+				const string InvalidSchemaWarning =
 					"\r\n" +
-					"{0}({1}): {2}\r\n" +
+					"{File}({Line}): {Message}\r\n" +
 					"XML schema failed to compile; validation of the final AppxManifest.xml will be skipped.\r\n" +
 					"If your AppxManifest.xml is valid then this is harmless, but if it contains invalid content you may encounter packaging or deployment errors.\r\n" +
 					"Updating your Windows SDK and/or Visual Studio installation may correct the schema problems and simplify diagnosis of invalid content.\r\n";
 
-				Log.TraceWarning(InvalidSchemaWarning, e.SourceUri, e.LineNumber, e.Message);
+				Logger.LogWarning(InvalidSchemaWarning, e.SourceUri, e.LineNumber, e.Message);
 				return;
 			}
 
@@ -2482,12 +2496,12 @@
 				switch (args.Severity)
 				{
 					case System.Xml.Schema.XmlSeverityType.Error:
-						Log.TraceError(args.Message);
+						Logger.LogError("{Message}", args.Message);
 						ValidationSucceeded = false;
 						break;
 
 					case System.Xml.Schema.XmlSeverityType.Warning:
-						Log.TraceWarning(args.Message);
+						Logger.LogWarning("{Message}", args.Message);
 						break;
 
 					default:
diff --git a/Engine/Platforms/Hololens/Source/Programs/UnrealBuildTool/UEBuildHoloLens.cs b/Engine/Platforms/Hololens/Source/Programs/UnrealBuildTool/UEBuildHoloLens.cs
--- a/Engine/Platforms/Hololens/Source/Programs/UnrealBuildTool/UEBuildHoloLens.cs
+++ b/Engine/Platforms/Hololens/Source/Programs/UnrealBuildTool/UEBuildHoloLens.cs
@@ -10,6 +10,7 @@
 using EpicGames.Core;
 using System.Text.RegularExpressions;
 using System.Runtime.Versioning;
+using Microsoft.Extensions.Logging;
 
 namespace UnrealBuildTool
 {
@@ -162,8 +163,8 @@
 		public static readonly Version MinimumSDKVersionRecommended = new Version(10, 0, 17763, 0);
 		public static readonly Version MaximumSDKVersionTested = new Version(10, 0, 18362, int.MaxValue);
 
-		public HoloLensPlatform(MicrosoftPlatformSDK InSDK) 
-			: base(UnrealTargetPlatform.HoloLens, InSDK)
+		public HoloLensPlatform(MicrosoftPlatformSDK InSDK, ILogger InLogger) 
+			: base(UnrealTargetPlatform.HoloLens, InSDK, InLogger)
 		{
 		}
 
@@ -206,7 +207,7 @@
 				}
 				else
 				{
-					Target.HoloLensPlatform.Compiler = WindowsPlatform.GetDefaultCompiler(Target.ProjectFile, Architecture);
+					Target.HoloLensPlatform.Compiler = WindowsPlatform.GetDefaultCompiler(Target.ProjectFile, Architecture, Logger);
 				}
 			}
 
@@ -239,7 +240,7 @@
 			// Be resilient to SDKs being uninstalled but still referenced in the INI file
 			VersionNumber? SelectedWindowsSdkVersion;
 			DirectoryReference? SelectedWindowsSdkDir;
-			if (WindowsPlatform.TryGetWindowsSdkDir(Target.HoloLensPlatform.Win10SDKVersionString, out SelectedWindowsSdkVersion, out SelectedWindowsSdkDir))
+			if (WindowsPlatform.TryGetWindowsSdkDir(Target.HoloLensPlatform.Win10SDKVersionString, Logger, out SelectedWindowsSdkVersion, out SelectedWindowsSdkDir))
 			{
 				Target.WindowsPlatform.WindowsSdkVersion = Target.HoloLensPlatform.Win10SDKVersionString;
 			}
@@ -258,15 +259,15 @@
 
 				if (Target.HoloLensPlatform.Win10SDKVersion < MinimumSDKVersionRecommended)
 				{
-					Log.TraceWarning("Your Windows SDK version {0} is older than the minimum recommended version ({1}) for HoloLens.  Consider upgrading.", Target.HoloLensPlatform.Win10SDKVersion, MinimumSDKVersionRecommended);
+					Logger.LogWarning("Your Windows SDK version {SdkVer} is older than the minimum recommended version ({RecommendedVer}) for HoloLens.  Consider upgrading.", Target.HoloLensPlatform.Win10SDKVersion, MinimumSDKVersionRecommended);
 				}
 				else if (Target.HoloLensPlatform.Win10SDKVersion > MaximumSDKVersionTested)
 				{
-					Log.TraceInformationOnce("Your Windows SDK version ({0}) for HoloLens is newer than the highest tested with this version of UBT ({1}).  This is probably fine, but if you encounter issues consider using an earlier SDK.", Target.HoloLensPlatform.Win10SDKVersion, MaximumSDKVersionTested);
+					Log.TraceInformationOnce("Your Windows SDK version ({SdkVer}) for HoloLens is newer than the highest tested with this version of UBT ({MaxTestedVer}).  This is probably fine, but if you encounter issues consider using an earlier SDK.", Target.HoloLensPlatform.Win10SDKVersion, MaximumSDKVersionTested);
 				}
 			}
 
-			HoloLensExports.InitWindowsSdkToolPath(Target.HoloLensPlatform.Win10SDKVersion.ToString());
+			HoloLensExports.InitWindowsSdkToolPath(Target.HoloLensPlatform.Win10SDKVersion.ToString(), Logger);
 
 			// ISPC currently doesn't support Windows-AArch64
 			Target.bCompileISPC = false;
@@ -299,12 +300,12 @@
 			return base.GetBinaryExtension(InBinaryType);
 		}
 
-		internal static DirectoryReference? GetCppCXMetadataLocation(WindowsCompiler Compiler, string CompilerVersion, WindowsArchitecture Architecture)
+		internal static DirectoryReference? GetCppCXMetadataLocation(WindowsCompiler Compiler, string CompilerVersion, WindowsArchitecture Architecture, ILogger Logger)
 		{
 			VersionNumber? SelectedToolChainVersion;
 			DirectoryReference? SelectedToolChainDir;
 			DirectoryReference? SelectedRedistDir;
-			if (!WindowsPlatform.TryGetToolChainDir(Compiler, CompilerVersion, Architecture, out SelectedToolChainVersion, out SelectedToolChainDir, out SelectedRedistDir))
+			if (!WindowsPlatform.TryGetToolChainDir(Compiler, CompilerVersion, Architecture, Logger, out SelectedToolChainVersion, out SelectedToolChainDir, out SelectedRedistDir))
 			{
 				return null;
 			}
@@ -347,11 +348,11 @@
 			return LatestVersion;
 		}
 
-		internal static string GetLatestMetadataPathForApiContract(string ApiContract, WindowsCompiler Compiler)
+		internal static string GetLatestMetadataPathForApiContract(string ApiContract, WindowsCompiler Compiler, ILogger Logger)
 		{
 			DirectoryReference? SDKFolder;
 			VersionNumber? SDKVersion;
-			if (!WindowsPlatform.TryGetWindowsSdkDir("Latest", out SDKVersion, out SDKFolder))
+			if (!WindowsPlatform.TryGetWindowsSdkDir("Latest", Logger, out SDKVersion, out SDKFolder))
 			{
 				return string.Empty;
 			}
@@ -462,7 +463,7 @@
 		/// <param name="Receipt">Information about the target being deployed</param>
 		public override void Deploy(TargetReceipt Receipt)
 		{
-			new HoloLensDeploy().PrepTargetForDeployment(Receipt);
+			new HoloLensDeploy(Logger).PrepTargetForDeployment(Receipt);
 		}
 
 		/// <summary>
@@ -516,7 +517,7 @@
 			}
 		}
 
-		internal static void ExpandWinMDReferences(ReadOnlyTargetRules Target, string SDKFolder, string SDKVersion, ref List<string> WinMDReferences)
+		internal static void ExpandWinMDReferences(ReadOnlyTargetRules Target, string SDKFolder, string SDKVersion, ref List<string> WinMDReferences, ILogger Logger)
 		{
 			// Code below will fail when not using the Win10 SDK.  Early out to avoid warning spam.
 			if (!Target.WindowsPlatform.bUseWindowsSDK10)
@@ -557,7 +558,7 @@
 						}
 						else
 						{
-							Log.TraceWarning("Unable to resolve location for HoloLens WinMD api contract {0}, file {1}", WinMDRef, ExpandedWinMDRef);
+							Logger.LogWarning("Unable to resolve location for HoloLens WinMD api contract {Contract}, file {File}", WinMDRef, ExpandedWinMDRef);
 						}
 					}
 				}
@@ -582,7 +583,7 @@
 			List<string> AlwaysReferenceContracts = new List<string>();
 			AlwaysReferenceContracts.Add("Windows.Foundation.FoundationContract");
 			AlwaysReferenceContracts.Add("Windows.Foundation.UniversalApiContract");
-			ExpandWinMDReferences(Target, Win10SDKRoot!, Target.HoloLensPlatform.Win10SDKVersion!.ToString(), ref AlwaysReferenceContracts);
+			ExpandWinMDReferences(Target, Win10SDKRoot!, Target.HoloLensPlatform.Win10SDKVersion!.ToString(), ref AlwaysReferenceContracts, Logger);
 
 			StringBuilder WinMDReferenceArguments = new StringBuilder();
 			foreach (string WinMDReference in AlwaysReferenceContracts)
@@ -731,7 +732,7 @@
 		/// <returns>New toolchain instance.</returns>
 		public override UEToolChain CreateToolChain(ReadOnlyTargetRules Target)
 		{
-			return new HoloLensToolChain(Target);
+			return new HoloLensToolChain(Target, Logger);
 		}
 
 		/// <summary>
@@ -757,14 +758,14 @@
 		/// <summary>
 		/// Register the platform with the UEBuildPlatform class
 		/// </summary>
-		public override void RegisterBuildPlatforms()
+		public override void RegisterBuildPlatforms(ILogger Logger)
 		{
 			if (OperatingSystem.IsWindows())
 			{
 				// for GetValidSoftwareVersionRange reasons we probably want a HoloLensePlatformSDK class
-				MicrosoftPlatformSDK SDK = new MicrosoftPlatformSDK();
+				MicrosoftPlatformSDK SDK = new MicrosoftPlatformSDK(Logger);
 
-				UEBuildPlatform.RegisterBuildPlatform(new HoloLensPlatform(SDK));
+				UEBuildPlatform.RegisterBuildPlatform(new HoloLensPlatform(SDK, Logger), Logger);
 				UEBuildPlatform.RegisterPlatformWithGroup(UnrealTargetPlatform.HoloLens, UnrealPlatformGroup.Microsoft);
 				UEBuildPlatform.RegisterPlatformWithGroup(UnrealTargetPlatform.HoloLens, UnrealPlatformGroup.HoloLens);
 			}
diff --git a/Engine/Platforms/Hololens/Source/Programs/UnrealBuildTool/WinMDRegistrationInfo.cs b/Engine/Platforms/Hololens/Source/Programs/UnrealBuildTool/WinMDRegistrationInfo.cs
--- a/Engine/Platforms/Hololens/Source/Programs/UnrealBuildTool/WinMDRegistrationInfo.cs
+++ b/Engine/Platforms/Hololens/Source/Programs/UnrealBuildTool/WinMDRegistrationInfo.cs
@@ -8,6 +8,7 @@
 using System.Runtime.InteropServices;
 using System.Runtime.Versioning;
 using EpicGames.Core;
+using Microsoft.Extensions.Logging;
 
 namespace UnrealBuildTool
 {
@@ -74,9 +75,9 @@
 
 			VersionNumber? OutSdkVersion;
 			DirectoryReference? OutSdkDir;
-			if (!WindowsPlatform.TryGetWindowsSdkDir(SdkVersion, out OutSdkVersion, out OutSdkDir))
+			if (!WindowsPlatform.TryGetWindowsSdkDir(SdkVersion, Logger, out OutSdkVersion, out OutSdkDir))
 			{
-				Log.TraceError("Failed to find WinSDK " + SdkVersion);
+				Logger.LogError("Failed to find WinSDK {SdkVersion}", SdkVersion);
 				throw new Exception("Failed to find WinSDK");
 			}
 
@@ -107,7 +108,7 @@
 					}
 					else
 					{
-						Log.TraceWarning("Unable to resolve location for HoloLens WinMD api contract {0}, file {1}", WinMDRef, ExpandedWinMDRef);
+						Logger.LogWarning("Unable to resolve location for HoloLens WinMD api contract {Contract}, file {File}", WinMDRef, ExpandedWinMDRef);
 					}
 				}
 			}
@@ -115,15 +116,20 @@
 			return ExpandedWinMDReferences;
 		}
 
+		private readonly ILogger Logger;
+
 		/// <summary>
 		/// WinMD reference info
 		/// </summary>
 		/// <param name="InWindMDSourcePath"></param>
 		/// <param name="InPackageRelativeDllPath"></param>
 		/// <param name="SdkVersion"></param>
+		/// <param name="InLogger"></param>
 		[SupportedOSPlatform("windows")]
-		public WinMDRegistrationInfo(FileReference InWindMDSourcePath, string InPackageRelativeDllPath, string SdkVersion)
+		public WinMDRegistrationInfo(FileReference InWindMDSourcePath, string InPackageRelativeDllPath, string SdkVersion, ILogger InLogger)
 		{
+			Logger = InLogger;
+
 			PackageRelativeDllPath = InPackageRelativeDllPath;
 			ResolveSearchPaths.Add(InWindMDSourcePath.Directory.FullName);
 			List<string> WinMDAssemblies = ExpandWinMDReferences(SdkVersion, new string[] {
diff --git a/Engine/Source/Programs/AutomationTool/Android/AndroidPlatform.Automation.cs b/Engine/Source/Programs/AutomationTool/Android/AndroidPlatform.Automation.cs
--- a/Engine/Source/Programs/AutomationTool/Android/AndroidPlatform.Automation.cs
+++ b/Engine/Source/Programs/AutomationTool/Android/AndroidPlatform.Automation.cs
@@ -3874,7 +3874,7 @@
 
 	public override void StripSymbols(FileReference SourceFile, FileReference TargetFile)
 	{
-		AndroidExports.StripSymbols(SourceFile, TargetFile);
+		AndroidExports.StripSymbols(SourceFile, TargetFile, Log.Logger);
 	}
 }
 
diff --git a/Engine/Source/Programs/AutomationTool/AutomationUtils/Automation.cs b/Engine/Source/Programs/AutomationTool/AutomationUtils/Automation.cs
--- a/Engine/Source/Programs/AutomationTool/AutomationUtils/Automation.cs
+++ b/Engine/Source/Programs/AutomationTool/AutomationUtils/Automation.cs
@@ -113,7 +113,7 @@
 				}
 				
 				// Initialize UBT
-				if (!UnrealBuildTool.PlatformExports.Initialize())
+				if (!UnrealBuildTool.PlatformExports.Initialize(Log.Logger))
 				{
 					Log.TraceInformation("Failed to initialize UBT");
 					return ExitCode.Error_Unknown;
diff --git a/Engine/Source/Programs/AutomationTool/AutomationUtils/ParallelExecutor.cs b/Engine/Source/Programs/AutomationTool/AutomationUtils/ParallelExecutor.cs
--- a/Engine/Source/Programs/AutomationTool/AutomationUtils/ParallelExecutor.cs
+++ b/Engine/Source/Programs/AutomationTool/AutomationUtils/ParallelExecutor.cs
@@ -78,7 +78,7 @@
 
 		public static int Execute(string ActionsFileName, bool bStopOnErrors)
 		{
-			return Execute(ActionsFileName, UnrealBuildTool.ParallelExecutorConfiguration.MaxParallelProcesses, bStopOnErrors);
+			return Execute(ActionsFileName, UnrealBuildTool.ParallelExecutorConfiguration.GetMaxParallelProcesses(Log.Logger), bStopOnErrors);
 		}
 
 		public static int Execute(string ActionsFileName, int MaxProcesses, bool bStopOnErrors)
diff --git a/Engine/Source/Programs/AutomationTool/AutomationUtils/ProjectUtils.cs b/Engine/Source/Programs/AutomationTool/AutomationUtils/ProjectUtils.cs
--- a/Engine/Source/Programs/AutomationTool/AutomationUtils/ProjectUtils.cs
+++ b/Engine/Source/Programs/AutomationTool/AutomationUtils/ProjectUtils.cs
@@ -201,7 +201,7 @@
             }
 
 			// Check if encryption or signing is enabled
-			EncryptionAndSigning.CryptoSettings Settings = EncryptionAndSigning.ParseCryptoSettings(RawProjectPath.Directory, Platform);
+			EncryptionAndSigning.CryptoSettings Settings = EncryptionAndSigning.ParseCryptoSettings(RawProjectPath.Directory, Platform, Log.Logger);
 			if (Settings.IsAnyEncryptionEnabled() || Settings.IsPakSigningEnabled())
 			{
 				OutReason = "encryption/signing is enabled";
@@ -770,7 +770,7 @@
 						typeof(UnrealBuildTool.PlatformExports).Assembly.Location
 					};
 			List<string> PreprocessorDefinitions = RulesAssembly.GetPreprocessorDefinitions();
-			Assembly TargetsDLL = DynamicCompilation.CompileAndLoadAssembly(TargetsDllFilename, new HashSet<FileReference>(TargetScripts), ReferencedAssemblies, PreprocessorDefinitions, DoNotCompile);
+			Assembly TargetsDLL = DynamicCompilation.CompileAndLoadAssembly(TargetsDllFilename, new HashSet<FileReference>(TargetScripts), Log.Logger, ReferencedAssemblies, PreprocessorDefinitions, DoNotCompile);
 			Type[] AllCompiledTypes = TargetsDLL.GetTypes();
 			foreach (Type TargetType in AllCompiledTypes)
 			{
@@ -897,7 +897,7 @@
 			}
 
 			// Search NativeProjects (sibling folders).
-			IEnumerable<FileReference> Projects = NativeProjects.EnumerateProjectFiles();
+			IEnumerable<FileReference> Projects = NativeProjects.EnumerateProjectFiles(Log.Logger);
 
 			FileReference ProjectPath = Projects.Where(R => string.Equals(R.GetFileName(), ProjectFile, StringComparison.OrdinalIgnoreCase)).FirstOrDefault();
 
@@ -1019,7 +1019,7 @@
 
         public BranchInfo()
         {
-            IEnumerable<FileReference> ProjectFiles = UnrealBuildTool.NativeProjects.EnumerateProjectFiles();
+            IEnumerable<FileReference> ProjectFiles = UnrealBuildTool.NativeProjects.EnumerateProjectFiles(Log.Logger);
 			foreach (FileReference InfoEntry in ProjectFiles)
 			{
 				AllProjects.Add(new BranchUProject(InfoEntry));
diff --git a/Engine/Source/Programs/AutomationTool/AutomationUtils/UnrealBuild.cs b/Engine/Source/Programs/AutomationTool/AutomationUtils/UnrealBuild.cs
--- a/Engine/Source/Programs/AutomationTool/AutomationUtils/UnrealBuild.cs
+++ b/Engine/Source/Programs/AutomationTool/AutomationUtils/UnrealBuild.cs
@@ -951,12 +951,12 @@
 
 		public bool CanUseXGE(UnrealTargetPlatform Platform)
 		{
-			return PlatformExports.CanUseXGE(Platform);
+			return PlatformExports.CanUseXGE(Platform, Log.Logger);
 		}
 
 		public bool CanUseParallelExecutor(UnrealTargetPlatform Platform)
 		{
-			return PlatformExports.CanUseParallelExecutor(Platform);
+			return PlatformExports.CanUseParallelExecutor(Platform, Log.Logger);
 		}
 
 		private bool ParseParam(string Name)
diff --git a/Engine/Source/Programs/AutomationTool/BuildGraph/Tasks/CommandletTask.cs b/Engine/Source/Programs/AutomationTool/BuildGraph/Tasks/CommandletTask.cs
--- a/Engine/Source/Programs/AutomationTool/BuildGraph/Tasks/CommandletTask.cs
+++ b/Engine/Source/Programs/AutomationTool/BuildGraph/Tasks/CommandletTask.cs
@@ -88,7 +88,7 @@
 				}
 				else
 				{
-					ProjectFile = NativeProjects.EnumerateProjectFiles().FirstOrDefault(x => x.GetFileNameWithoutExtension().Equals(Parameters.Project, StringComparison.OrdinalIgnoreCase));
+					ProjectFile = NativeProjects.EnumerateProjectFiles(Log.Logger).FirstOrDefault(x => x.GetFileNameWithoutExtension().Equals(Parameters.Project, StringComparison.OrdinalIgnoreCase));
 				}
 
 				if(ProjectFile == null || !FileReference.Exists(ProjectFile))
diff --git a/Engine/Source/Programs/AutomationTool/BuildGraph/Tasks/CompileTask.cs b/Engine/Source/Programs/AutomationTool/BuildGraph/Tasks/CompileTask.cs
--- a/Engine/Source/Programs/AutomationTool/BuildGraph/Tasks/CompileTask.cs
+++ b/Engine/Source/Programs/AutomationTool/BuildGraph/Tasks/CompileTask.cs
@@ -232,7 +232,7 @@
 				}
 				else
 				{
-					ProjectFile = NativeProjects.EnumerateProjectFiles().FirstOrDefault(x => x.GetFileNameWithoutExtension().Equals(Parameters.Project, StringComparison.OrdinalIgnoreCase));
+					ProjectFile = NativeProjects.EnumerateProjectFiles(Log.Logger).FirstOrDefault(x => x.GetFileNameWithoutExtension().Equals(Parameters.Project, StringComparison.OrdinalIgnoreCase));
 				}
 
 				if(ProjectFile == null || !FileReference.Exists(ProjectFile))
diff --git a/Engine/Source/Programs/AutomationTool/Gauntlet/Platform/Android/Gauntlet.AndroidPGO.cs b/Engine/Source/Programs/AutomationTool/Gauntlet/Platform/Android/Gauntlet.AndroidPGO.cs
--- a/Engine/Source/Programs/AutomationTool/Gauntlet/Platform/Android/Gauntlet.AndroidPGO.cs
+++ b/Engine/Source/Programs/AutomationTool/Gauntlet/Platform/Android/Gauntlet.AndroidPGO.cs
@@ -44,7 +44,7 @@
 				MergeCommandBuilder.AppendFormat(" \"{0}\"", ProfRawFile);
 			}
 
-			int ReturnCode = UnrealBuildTool.Utils.RunLocalProcessAndLogOutput(new ProcessStartInfo(AndroidLlvmProfdataPath, string.Format("merge{0} -o \"{1}\"", MergeCommandBuilder, LocalProfDataFile)));
+			int ReturnCode = UnrealBuildTool.Utils.RunLocalProcessAndLogOutput(new ProcessStartInfo(AndroidLlvmProfdataPath, string.Format("merge{0} -o \"{1}\"", MergeCommandBuilder, LocalProfDataFile)), EpicGames.Core.Log.Logger);
 			if (ReturnCode != 0)
 			{
 				throw new AutomationException(string.Format("{0} failed to merge profraw data. Error code {1}. ({2} {3})", Path.GetFileName(AndroidLlvmProfdataPath), ReturnCode, MergeCommandBuilder, LocalProfDataFile));
diff --git a/Engine/Source/Programs/AutomationTool/IOS/IOSPlatform.Automation.cs b/Engine/Source/Programs/AutomationTool/IOS/IOSPlatform.Automation.cs
--- a/Engine/Source/Programs/AutomationTool/IOS/IOSPlatform.Automation.cs
+++ b/Engine/Source/Programs/AutomationTool/IOS/IOSPlatform.Automation.cs
@@ -308,7 +308,7 @@
 			// if non-interactive, we can just run directly in the current shell
 			if (!bInteractive)
 			{
-				UnrealBuildTool.Utils.RunLocalProcessAndReturnStdOut(Command, Params, out ExitCode, true);
+				UnrealBuildTool.Utils.RunLocalProcessAndReturnStdOut(Command, Params, Log.Logger, out ExitCode);
 			}
 			else
 			{
@@ -338,7 +338,7 @@
 
 				Console.WriteLine("\n\n\n{0}\n\n\n", Params);
 
-				UnrealBuildTool.Utils.RunLocalProcessAndReturnStdOut("osascript", Params, out ExitCode, true);
+				UnrealBuildTool.Utils.RunLocalProcessAndReturnStdOut("osascript", Params, Log.Logger, out ExitCode);
 				if (ExitCode == 0)
 				{
 					ExitCode = int.Parse(File.ReadAllText(ReturnCodeFilename));
@@ -621,7 +621,7 @@
 	{
 		FileReference TargetReceiptFileName = GetTargetReceiptFileName(Config, InExecutablePath, InEngineDir, InProjectDirectory, bIsUEGame);
 
-		return IOSExports.PrepForUATPackageOrDeploy(Config, ProjectFile, InProjectName, InProjectDirectory, InExecutablePath, InEngineDir, bForDistribution, CookFlavor, bIsDataDeploy, bCreateStubIPA, TargetReceiptFileName);
+		return IOSExports.PrepForUATPackageOrDeploy(Config, ProjectFile, InProjectName, InProjectDirectory, InExecutablePath, InEngineDir, bForDistribution, CookFlavor, bIsDataDeploy, bCreateStubIPA, TargetReceiptFileName, Log.Logger);
 	}
 
 
@@ -648,7 +648,7 @@
 	public virtual bool DeployGeneratePList(FileReference ProjectFile, UnrealTargetConfiguration Config, DirectoryReference ProjectDirectory, bool bIsUEGame, string GameName, bool bIsClient, string ProjectName, DirectoryReference InEngineDir, DirectoryReference AppDirectory, string InExecutablePath, out bool bSupportsPortrait, out bool bSupportsLandscape)
 	{
 		FileReference TargetReceiptFileName = GetTargetReceiptFileName(Config, InExecutablePath, InEngineDir, ProjectDirectory, bIsUEGame);
-		return IOSExports.GeneratePList(ProjectFile, Config, ProjectDirectory, bIsUEGame, GameName, bIsClient, ProjectName, InEngineDir, AppDirectory, TargetReceiptFileName, out bSupportsPortrait, out bSupportsLandscape);
+		return IOSExports.GeneratePList(ProjectFile, Config, ProjectDirectory, bIsUEGame, GameName, bIsClient, ProjectName, InEngineDir, AppDirectory, TargetReceiptFileName, Log.Logger, out bSupportsPortrait, out bSupportsLandscape);
 	}
 
 	protected string MakeIPAFileName(UnrealTargetConfiguration TargetConfiguration, ProjectParams Params, DeploymentContext SC, bool bAllowDistroPrefix)
@@ -875,7 +875,7 @@
 
 		StageLaunchScreenStoryboard(Params, SC);
 
-		IOSExports.GenerateAssetCatalog(Params.RawProjectPath, new FileReference(FullExePath), new DirectoryReference(CombinePaths(Params.BaseStageDirectory, (TargetPlatformType == UnrealTargetPlatform.IOS ? "IOS" : "TVOS"))), TargetPlatformType);
+		IOSExports.GenerateAssetCatalog(Params.RawProjectPath, new FileReference(FullExePath), new DirectoryReference(CombinePaths(Params.BaseStageDirectory, (TargetPlatformType == UnrealTargetPlatform.IOS ? "IOS" : "TVOS"))), TargetPlatformType, Log.Logger);
 
 		bCreatedIPA = false;
 		bool bNeedsIPA = false;
@@ -1704,7 +1704,7 @@
 				string IdeviceFSArgs = "-b " + "\"" + BundleIdentifier + " -x " + Directory.GetCurrentDirectory() + "\\CommandsToPush.txt -u " + "\"" + Params.DeviceNames[0];
 				IdeviceFSArgs = GetLibimobileDeviceNetworkedArgument(IdeviceFSArgs, Params.DeviceNames[0]);
 
-				Utils.RunLocalProcessAndReturnStdOut(DeviceFS, IdeviceFSArgs, out ExitCode, true);
+				Utils.RunLocalProcessAndReturnStdOut(DeviceFS, IdeviceFSArgs, Log.Logger, out ExitCode);
 				if (ExitCode != 0)
 				{
 					throw new AutomationException("Failed to deploy manifest to mobile device.");
@@ -2256,7 +2256,7 @@
 
 	public override void StripSymbols(FileReference SourceFile, FileReference TargetFile)
 	{
-		IOSExports.StripSymbols(PlatformType, SourceFile, TargetFile);
+		IOSExports.StripSymbols(PlatformType, SourceFile, TargetFile, Log.Logger);
 	}
 
 
diff --git a/Engine/Source/Programs/AutomationTool/Linux/LinuxPlatform.Automation.cs b/Engine/Source/Programs/AutomationTool/Linux/LinuxPlatform.Automation.cs
--- a/Engine/Source/Programs/AutomationTool/Linux/LinuxPlatform.Automation.cs
+++ b/Engine/Source/Programs/AutomationTool/Linux/LinuxPlatform.Automation.cs
@@ -410,7 +410,7 @@
 
 	public override void StripSymbols(FileReference SourceFile, FileReference TargetFile)
 	{
-		LinuxExports.StripSymbols(SourceFile, TargetFile);
+		LinuxExports.StripSymbols(SourceFile, TargetFile, Log.Logger);
 	}
 }
 
diff --git a/Engine/Source/Programs/AutomationTool/Mac/MacPlatform.Automation.cs b/Engine/Source/Programs/AutomationTool/Mac/MacPlatform.Automation.cs
--- a/Engine/Source/Programs/AutomationTool/Mac/MacPlatform.Automation.cs
+++ b/Engine/Source/Programs/AutomationTool/Mac/MacPlatform.Automation.cs
@@ -689,6 +689,6 @@
 
 	public override void StripSymbols(FileReference SourceFile, FileReference TargetFile)
 	{
-		MacExports.StripSymbols(SourceFile, TargetFile);
+		MacExports.StripSymbols(SourceFile, TargetFile, Log.Logger);
 	}
 }
diff --git a/Engine/Source/Programs/AutomationTool/Program.cs b/Engine/Source/Programs/AutomationTool/Program.cs
--- a/Engine/Source/Programs/AutomationTool/Program.cs
+++ b/Engine/Source/Programs/AutomationTool/Program.cs
@@ -397,7 +397,7 @@
 				bool bWaitForUATMutex = AutomationToolCommandLine.IsSetGlobal("-WaitForUATMutex");
 
 				// Don't allow simultaneous execution of AT (in the same branch)
-				ReturnCode = await ProcessSingleton.RunSingleInstanceAsync(MainProc, bWaitForUATMutex);
+				ReturnCode = await ProcessSingleton.RunSingleInstanceAsync(MainProc, bWaitForUATMutex, Log.Logger);
 			}
 			catch (Exception Ex)
             {
@@ -434,7 +434,8 @@
 					out bBuildSuccess, (int Count) =>
                     {
 						Logger.LogInformation("Building {Count} projects (see Log 'Engine/Programs/AutomationTool/Saved/Logs/Log.txt' for more details)", Count);
-					});
+					},
+					Log.Logger);
 
 			if (!bBuildSuccess)
             {
diff --git a/Engine/Source/Programs/AutomationTool/Scripts/BuildPhysX.Automation.cs b/Engine/Source/Programs/AutomationTool/Scripts/BuildPhysX.Automation.cs
--- a/Engine/Source/Programs/AutomationTool/Scripts/BuildPhysX.Automation.cs
+++ b/Engine/Source/Programs/AutomationTool/Scripts/BuildPhysX.Automation.cs
@@ -282,7 +282,7 @@
 			StartInfo.WorkingDirectory = CMakeTargetDirectory.FullName;
 			StartInfo.Arguments = GetCMakeArguments(TargetLib, TargetConfiguration);
 
-			if (Utils.RunLocalProcessAndLogOutput(StartInfo) != 0)
+			if (Utils.RunLocalProcessAndLogOutput(StartInfo, Log.Logger) != 0)
 			{
 				throw new AutomationException("Unable to generate projects for {0}.", TargetLib.ToString() + ", " + FriendlyName);
 			}
@@ -399,7 +399,7 @@
 			LogInformation("Working in: {0}", StartInfo.WorkingDirectory);
 			LogInformation("{0} {1}", StartInfo.FileName, StartInfo.Arguments);
 
-			if (Utils.RunLocalProcessAndLogOutput(StartInfo) != 0)
+			if (Utils.RunLocalProcessAndLogOutput(StartInfo, Log.Logger) != 0)
 			{
 				throw new AutomationException("Unabled to build {0}. Build process failed.", Makefile);
 			}
@@ -1083,14 +1083,14 @@
 			StartInfo.RedirectStandardError = true;
 
 			LogInformation("Running: '{0} {1}'", StartInfo.FileName, StartInfo.Arguments);
-			Utils.RunLocalProcessAndLogOutput(StartInfo);
+			Utils.RunLocalProcessAndLogOutput(StartInfo, Log.Logger);
 
 			// BreakpadSymbolEncoder
 			StartInfo.FileName = BreakpadSymbolEncoderPath.FullName + ExeSuffix;
 			StartInfo.Arguments = PSymbolFile.FullName + " " + SymbolFile.FullName;
 
 			LogInformation("Running: '{0} {1}'", StartInfo.FileName, StartInfo.Arguments);
-			Utils.RunLocalProcessAndLogOutput(StartInfo);
+			Utils.RunLocalProcessAndLogOutput(StartInfo, Log.Logger);
 
 			// Clean up the Temp *.psym file, as they are no longer needed
 			InternalUtils.SafeDeleteFile(PSymbolFile.FullName);
@@ -1104,7 +1104,7 @@
 					StrippedFile.FullName;
 
 				LogInformation("Running: '{0} {1}'", StartInfo.FileName, StartInfo.Arguments);
-				Utils.RunLocalProcessAndLogOutput(StartInfo);
+				Utils.RunLocalProcessAndLogOutput(StartInfo, Log.Logger);
 
 				// objcopy --only-keep-debug sofile.so sofile.debug
 				StartInfo.FileName = ObjcopyPath.FullName + ExeSuffix;
@@ -1113,7 +1113,7 @@
 					DebugFile.FullName;
 
 				LogInformation("Running: '{0} {1}'", StartInfo.FileName, StartInfo.Arguments);
-				Utils.RunLocalProcessAndLogOutput(StartInfo);
+				Utils.RunLocalProcessAndLogOutput(StartInfo, Log.Logger);
 
 				// objcopy --add-gnu-debuglink=sofile.debug sofile_stripped sofile.so
 				StartInfo.FileName = ObjcopyPath.FullName + ExeSuffix;
@@ -1123,7 +1123,7 @@
 					SOFile.FullName;
 
 				LogInformation("Running: '{0} {1}'", StartInfo.FileName, StartInfo.Arguments);
-				Utils.RunLocalProcessAndLogOutput(StartInfo);
+				Utils.RunLocalProcessAndLogOutput(StartInfo, Log.Logger);
 
 				GeneratedDebugSymbols.Add(SOFile.FullName, true);
 			}
@@ -1292,7 +1292,7 @@
 			StartInfo.RedirectStandardError = true;
 
 			LogInformation("Running: 'lipo {0}'", StartInfo.Arguments);
-			if (Utils.RunLocalProcessAndLogOutput(StartInfo) != 0)
+			if (Utils.RunLocalProcessAndLogOutput(StartInfo, Log.Logger) != 0)
 			{
 				LogError("Failed to create universal binary for {0}", LibFile);
 			}
diff --git a/Engine/Source/Programs/AutomationTool/Scripts/CopyBuildToStagingDirectory.Automation.cs b/Engine/Source/Programs/AutomationTool/Scripts/CopyBuildToStagingDirectory.Automation.cs
--- a/Engine/Source/Programs/AutomationTool/Scripts/CopyBuildToStagingDirectory.Automation.cs
+++ b/Engine/Source/Programs/AutomationTool/Scripts/CopyBuildToStagingDirectory.Automation.cs
@@ -2205,7 +2205,7 @@ namespace AutomationScripts
 				UnrealPakResponseFile.Remove(FileToRemove);
 			}
 
-			EncryptionAndSigning.CryptoSettings PakCryptoSettings = EncryptionAndSigning.ParseCryptoSettings(DirectoryReference.FromFile(Params.RawProjectPath), SC.StageTargetPlatform.IniPlatformType);
+			EncryptionAndSigning.CryptoSettings PakCryptoSettings = EncryptionAndSigning.ParseCryptoSettings(DirectoryReference.FromFile(Params.RawProjectPath), SC.StageTargetPlatform.IniPlatformType, Log.Logger);
 			FileReference CryptoKeysCacheFilename = FileReference.Combine(SC.MetadataDir, "Crypto.json");
 			PakCryptoSettings.Save(CryptoKeysCacheFilename);
 
@@ -3594,7 +3594,7 @@ namespace AutomationScripts
 			var TmpPackagingPath = GetTmpPackagingPath(Params, SC);
 
 			// Parse and cache crypto settings from INI file
-			EncryptionAndSigning.CryptoSettings PakCryptoSettings = EncryptionAndSigning.ParseCryptoSettings(DirectoryReference.FromFile(Params.RawProjectPath), SC.StageTargetPlatform.IniPlatformType);
+			EncryptionAndSigning.CryptoSettings PakCryptoSettings = EncryptionAndSigning.ParseCryptoSettings(DirectoryReference.FromFile(Params.RawProjectPath), SC.StageTargetPlatform.IniPlatformType, Log.Logger);
 			FileReference CryptoKeysCacheFilename = FileReference.Combine(SC.MetadataDir, "Crypto.json");
 			PakCryptoSettings.Save(CryptoKeysCacheFilename);
 
diff --git a/Engine/Source/Programs/AutomationTool/Scripts/ListThirdPartySoftware.Automation.cs b/Engine/Source/Programs/AutomationTool/Scripts/ListThirdPartySoftware.Automation.cs
--- a/Engine/Source/Programs/AutomationTool/Scripts/ListThirdPartySoftware.Automation.cs
+++ b/Engine/Source/Programs/AutomationTool/Scripts/ListThirdPartySoftware.Automation.cs
@@ -94,7 +94,7 @@
 
 			// Get the platforms to exclude
 			List<UnrealTargetPlatform> SupportedPlatforms = new List<UnrealTargetPlatform> { UnrealTargetPlatform.Parse(Object.GetStringField("Platform")) };
-			string[] ExcludePlatformNames = Utils.MakeListOfUnsupportedPlatforms(SupportedPlatforms, bIncludeUnbuildablePlatforms: true).ToArray();
+			string[] ExcludePlatformNames = Utils.MakeListOfUnsupportedPlatforms(SupportedPlatforms, bIncludeUnbuildablePlatforms: true, Log.Logger).ToArray();
 
 			// Find all the TPS files under the engine directory which match
 			foreach(DirectoryReference DirectoryToScan in SortedDirectoriesToScan)
diff --git a/Engine/Source/Programs/AutomationTool/TVOS/TVOSPlatform.Automation.cs b/Engine/Source/Programs/AutomationTool/TVOS/TVOSPlatform.Automation.cs
--- a/Engine/Source/Programs/AutomationTool/TVOS/TVOSPlatform.Automation.cs
+++ b/Engine/Source/Programs/AutomationTool/TVOS/TVOSPlatform.Automation.cs
@@ -29,7 +29,7 @@
 		{
 			TargetReceiptFileName = TargetReceipt.GetDefaultPath(InProjectDirectory, TargetName, UnrealTargetPlatform.TVOS, Config, "");
 		}
-		return TVOSExports.PrepForUATPackageOrDeploy(Config, ProjectFile, InProjectName, InProjectDirectory, InExecutablePath, InEngineDir, bForDistribution, CookFlavor, bIsDataDeploy, bCreateStubIPA, TargetReceiptFileName);
+		return TVOSExports.PrepForUATPackageOrDeploy(Config, ProjectFile, InProjectName, InProjectDirectory, InExecutablePath, InEngineDir, bForDistribution, CookFlavor, bIsDataDeploy, bCreateStubIPA, TargetReceiptFileName, Log.Logger);
 	}
 
     public override void GetProvisioningData(FileReference InProject, bool bDistribution, out string MobileProvision, out string SigningCertificate, out string Team, out bool bAutomaticSigning)
@@ -49,7 +49,7 @@
 		{
 			TargetReceiptFileName = TargetReceipt.GetDefaultPath(ProjectDirectory, TargetName, UnrealTargetPlatform.TVOS, Config, "");
 		}
-		return TVOSExports.GeneratePList(ProjectFile, Config, ProjectDirectory, bIsUEGame, GameName, bIsClient, ProjectName, InEngineDir, AppDirectory, TargetReceiptFileName, out bSupportsPortrait, out bSupportsLandscape);
+		return TVOSExports.GeneratePList(ProjectFile, Config, ProjectDirectory, bIsUEGame, GameName, bIsClient, ProjectName, InEngineDir, AppDirectory, TargetReceiptFileName, Log.Logger, out bSupportsPortrait, out bSupportsLandscape);
 	}
 
     public override string GetCookPlatform(bool bDedicatedServer, bool bIsClientOnly)
diff --git a/Engine/Source/Programs/AutomationTool/Turnkey/Commands/CreateBuild.cs b/Engine/Source/Programs/AutomationTool/Turnkey/Commands/CreateBuild.cs
--- a/Engine/Source/Programs/AutomationTool/Turnkey/Commands/CreateBuild.cs
+++ b/Engine/Source/Programs/AutomationTool/Turnkey/Commands/CreateBuild.cs
@@ -789,7 +789,7 @@
 							TurnkeyUtils.Log("Further, you can edit this in the Packaging Settings in the editor as needed");
 
 							// add this setting in to the ProjectSettings section
-							ConfigCache.WriteSettingToDefaultConfig(ConfigHierarchyType.Game, ProjectFile.Directory, ConfigCache.ConfigDefaultUpdateType.AddArrayEntry, SettingsSection, "ProjectCustomBuilds", ValueText);
+							ConfigCache.WriteSettingToDefaultConfig(ConfigHierarchyType.Game, ProjectFile.Directory, ConfigCache.ConfigDefaultUpdateType.AddArrayEntry, SettingsSection, "ProjectCustomBuilds", ValueText, Log.Logger);
 							ConfigCache.InvalidateCaches();
 						}
 					}
diff --git a/Engine/Source/Programs/AutomationTool/Turnkey/CopyProvider.cs b/Engine/Source/Programs/AutomationTool/Turnkey/CopyProvider.cs
--- a/Engine/Source/Programs/AutomationTool/Turnkey/CopyProvider.cs
+++ b/Engine/Source/Programs/AutomationTool/Turnkey/CopyProvider.cs
@@ -10,6 +10,7 @@
 using AutomationTool;
 using Ionic.Zip;
 using UnrealBuildTool;
+using EpicGames.Core;
 
 namespace Turnkey
 {
@@ -167,7 +168,7 @@
 							TurnkeyUtils.Log("7zip decompressing {0} to {1}...", OutputPath, DecompressLocation);
 
 							int ExitCode = UnrealBuildTool.Utils.RunLocalProcessAndLogOutput(TurnkeyUtils.ExpandVariables("$(EngineDir)/Restricted/NotForLicensees/Extras/ThirdPartyNotUE/7-Zip/7z.exe"),
-								string.Format("x -o{0} {1}", DecompressLocation, OutputPath));
+								string.Format("x -o{0} {1}", DecompressLocation, OutputPath), Log.Logger);
 							if (ExitCode != 0)
 							{
 								TurnkeyUtils.Log("Failed to uncompress a .7z file {0}", OutputPath);
diff --git a/Engine/Source/Programs/AutomationTool/Turnkey/TurnkeyManifest.cs b/Engine/Source/Programs/AutomationTool/Turnkey/TurnkeyManifest.cs
--- a/Engine/Source/Programs/AutomationTool/Turnkey/TurnkeyManifest.cs
+++ b/Engine/Source/Programs/AutomationTool/Turnkey/TurnkeyManifest.cs
@@ -221,7 +221,7 @@
 //					TurnkeyUtils.Log("Info: Processing manifest: {0}", LocalManifestPath);
 
 					// read in the .xml
-					TurnkeyManifest Manifest = Utils.ReadClass<TurnkeyManifest>(LocalManifestPath);
+					TurnkeyManifest Manifest = Utils.ReadClass<TurnkeyManifest>(LocalManifestPath, Log.Logger);
 					Manifest.PostDeserialize();
 
 					Manifests.Add(Manifest);
@@ -245,7 +245,7 @@
 
 		public void Write(string Path)
 		{
-			Utils.WriteClass<TurnkeyManifest>(this, Path, "");
+			Utils.WriteClass<TurnkeyManifest>(this, Path, "", Log.Logger);
 		}
 	}
 }
diff --git a/Engine/Source/Programs/AutomationTool/Turnkey/TurnkeyUtils.cs b/Engine/Source/Programs/AutomationTool/Turnkey/TurnkeyUtils.cs
--- a/Engine/Source/Programs/AutomationTool/Turnkey/TurnkeyUtils.cs
+++ b/Engine/Source/Programs/AutomationTool/Turnkey/TurnkeyUtils.cs
@@ -227,7 +227,7 @@
 			}
 
 			// sort the discovered projects
-			foreach (FileReference NativeProject in NativeProjects.EnumerateProjectFiles())
+			foreach (FileReference NativeProject in NativeProjects.EnumerateProjectFiles(EpicGames.Core.Log.Logger))
 			{
 				DirectoryReference ProjectDir = NativeProject.Directory;
 				ProjectMode Mode = ProjectMode.Normal;
diff --git a/Engine/Source/Programs/AutomationTool/Win/WinPlatform.Automation.cs b/Engine/Source/Programs/AutomationTool/Win/WinPlatform.Automation.cs
--- a/Engine/Source/Programs/AutomationTool/Win/WinPlatform.Automation.cs
+++ b/Engine/Source/Programs/AutomationTool/Win/WinPlatform.Automation.cs
@@ -450,7 +450,7 @@
 		StartInfo.Arguments = String.Format("\"{0}\" \"{1}\" -p", SourceFile.FullName, TargetFile.FullName);
 		StartInfo.UseShellExecute = false;
 		StartInfo.CreateNoWindow = true;
-		Utils.RunLocalProcessAndLogOutput(StartInfo);
+		Utils.RunLocalProcessAndLogOutput(StartInfo, Log.Logger);
 
 		if (bStripInPlace)
 		{
@@ -489,7 +489,7 @@
 				StartInfo.Arguments = string.Format("add /f \"@{0}\" /s \"{1}\" /t \"{2}\"", TempFileName, TempSymStoreDir, Product);
 				StartInfo.UseShellExecute = false;
 				StartInfo.CreateNoWindow = true;
-				if (Utils.RunLocalProcessAndLogOutput(StartInfo) != 0)
+				if (Utils.RunLocalProcessAndLogOutput(StartInfo, Log.Logger) != 0)
 				{
 					return false;
 				}
diff --git a/Engine/Source/Programs/Shared/EpicGames.Build/Automation/ProcessSingleton.cs b/Engine/Source/Programs/Shared/EpicGames.Build/Automation/ProcessSingleton.cs
--- a/Engine/Source/Programs/Shared/EpicGames.Build/Automation/ProcessSingleton.cs
+++ b/Engine/Source/Programs/Shared/EpicGames.Build/Automation/ProcessSingleton.cs
@@ -1,6 +1,7 @@
 ﻿// Copyright Epic Games, Inc. All Rights Reserved.
 
 using EpicGames.Core;
+using Microsoft.Extensions.Logging;
 using System;
 using System.Collections.Generic;
 using System.IO;
@@ -30,17 +31,17 @@
 		/// </summary>
 		/// <param name="Main"></param>
 		/// <param name="bWaitForUATMutex"></param>
-		public static async Task<ExitCode> RunSingleInstanceAsync(Func<Task<ExitCode>> Main, bool bWaitForUATMutex)
+		public static async Task<ExitCode> RunSingleInstanceAsync(Func<Task<ExitCode>> Main, bool bWaitForUATMutex, ILogger Logger)
 		{
 			// Need to execute this logic on a background thread, since mutex ownership on Linux has thread affinity (ie. mutexes be released on the
 			// same thread that acquires it, which is not guaranteed through an async continuation)
 			TaskCompletionSource<ExitCode> Result = new TaskCompletionSource<ExitCode>();
-			Thread Thread = new Thread(() => RunSingleInstanceThread(Main, Result, bWaitForUATMutex));
+			Thread Thread = new Thread(() => RunSingleInstanceThread(Main, Result, bWaitForUATMutex, Logger));
 			Thread.Start();
 			return await Result.Task;
 		}
 
-		public static void RunSingleInstanceThread(Func<Task<ExitCode>> Main, TaskCompletionSource<ExitCode> Result, bool bWaitForUATMutex)
+		public static void RunSingleInstanceThread(Func<Task<ExitCode>> Main, TaskCompletionSource<ExitCode> Result, bool bWaitForUATMutex, ILogger Logger)
 		{
 			try
 			{
@@ -57,12 +58,12 @@
 					{
 						if (bWaitForUATMutex)
 						{
-							Log.TraceWarning("Another instance of UAT at '{0}' is running, and the -WaitForUATMutex parameter has been used. Waiting for other UAT to finish...", EntryAssemblyLocation);
+							Logger.LogWarning("Another instance of UAT at '{File}' is running, and the -WaitForUATMutex parameter has been used. Waiting for other UAT to finish...", EntryAssemblyLocation);
 							int Seconds = 0;
 							while (WaitMutexNoExceptions(SingleInstanceMutex, 15 * 1000) == false)
 							{
 								Seconds += 15;
-								Log.TraceInformation("Still waiting for Mutex. {0} seconds has passed...", Seconds);
+								Logger.LogInformation("Still waiting for Mutex. {TimeSeconds} seconds has passed...", Seconds);
 							}
 						}
 						else
diff --git a/Engine/Source/Programs/Shared/EpicGames.Build/Platform/Mac/ApplePlatformSDK.cs b/Engine/Source/Programs/Shared/EpicGames.Build/Platform/Mac/ApplePlatformSDK.cs
--- a/Engine/Source/Programs/Shared/EpicGames.Build/Platform/Mac/ApplePlatformSDK.cs
+++ b/Engine/Source/Programs/Shared/EpicGames.Build/Platform/Mac/ApplePlatformSDK.cs
@@ -7,6 +7,7 @@
 using System.Text;
 using System.Text.RegularExpressions;
 using EpicGames.Core;
+using Microsoft.Extensions.Logging;
 using Microsoft.Win32;
 
 namespace UnrealBuildBase
@@ -137,10 +138,9 @@
 		/// <returns>The entire StdOut generated from the process as a single trimmed string</returns>
 		/// <param name="Command">Command to run</param>
 		/// <param name="Args">Arguments to Command</param>
-		private static string RunLocalProcessAndReturnStdOut(string Command, string Args)
+		private static string RunLocalProcessAndReturnStdOut(string Command, string Args, ILogger? Logger = null)
 		{
-			int ExitCode;
-			return RunLocalProcessAndReturnStdOut(Command, Args, out ExitCode);	
+			return RunLocalProcessAndReturnStdOut(Command, Args, out _, Logger);	
 		}
 
 		/// <summary>
@@ -150,8 +150,8 @@
 		/// <param name="Command">Command to run</param>
 		/// <param name="Args">Arguments to Command</param>
 		/// <param name="ExitCode">The return code from the process after it exits</param>
-		/// <param name="LogOutput">Whether to also log standard output and standard error</param>
-		private static string RunLocalProcessAndReturnStdOut(string Command, string? Args, out int ExitCode, bool LogOutput = false)
+		/// <param name="Logger">Whether to also log standard output and standard error</param>
+		private static string RunLocalProcessAndReturnStdOut(string Command, string? Args, out int ExitCode, ILogger? Logger = null)
 		{
 			// Process Arguments follow windows conventions in .NET Core
 			// Which means single quotes ' are not considered quotes.
@@ -181,16 +181,16 @@
 					StreamReader ErrorReader = LocalProcess.StandardError;
 					// trim off any extraneous new lines, helpful for those one-line outputs
 					ErrorOutput = ErrorReader.ReadToEnd().Trim();
-					if (LogOutput)
+					if (Logger != null)
 					{
 						if (FullOutput.Length > 0)
 						{
-							Log.TraceInformation(FullOutput);
+							Logger.LogInformation("{Message}", FullOutput);
 						}
 
 						if (ErrorOutput.Length > 0)
 						{
-							Log.TraceError(ErrorOutput);
+							Logger.LogError("{Message}", ErrorOutput);
 						}
 					}
 
