Remove HoloLens platform references from UnrealBuildTool.

#rnx

diff --git a/Engine/Platforms/Hololens/Source/Programs/UnrealBuildTool/UEBuildHoloLens.cs b/Engine/Platforms/Hololens/Source/Programs/UnrealBuildTool/UEBuildHoloLens.cs
--- a/Engine/Platforms/Hololens/Source/Programs/UnrealBuildTool/UEBuildHoloLens.cs
+++ b/Engine/Platforms/Hololens/Source/Programs/UnrealBuildTool/UEBuildHoloLens.cs
@@ -157,6 +157,41 @@ namespace UnrealBuildTool
 		}
 	}
 
+	[SupportedOSPlatform("windows")]
+	class HoloLensEnvironment : VCEnvironment
+	{
+		internal HoloLensEnvironment(VCEnvironmentParameters Params, ILogger Logger)
+			: base(Params, Logger)
+		{
+		}
+
+		/// <summary>
+		/// Return the standard Visual C++ library path for the given platform in this toolchain
+		/// </summary>
+		protected override DirectoryReference GetToolChainLibsDir(UnrealTargetPlatform Platform)
+		{
+			string ArchFolder = WindowsExports.GetArchitectureSubpath(Architecture);
+
+			// Add the standard Visual C++ library paths
+			if (ToolChain.IsMSVC())
+			{
+				return DirectoryReference.Combine(ToolChainDir, "lib", ArchFolder, "store");
+			}
+			else
+			{
+				DirectoryReference LibsPath = DirectoryReference.Combine(ToolChainDir, "LIB", "store");
+
+				if (Architecture == WindowsArchitecture.x64)
+				{
+					LibsPath = DirectoryReference.Combine(LibsPath, "amd64");
+				}
+
+				return LibsPath;
+			}
+		}
+	};
+
+
 	[SupportedOSPlatform("windows")]
 	class HoloLensPlatform : WindowsPlatform
 	{
@@ -168,6 +203,12 @@ namespace UnrealBuildTool
 		{
 		}
 
+		protected override VCEnvironment CreateVCEnvironment(TargetRules Target)
+		{
+			VCEnvironmentParameters Params = new VCEnvironmentParameters(Target.WindowsPlatform.Compiler, Platform, Target.WindowsPlatform.Architecture, Target.WindowsPlatform.CompilerVersion, Target.WindowsPlatform.WindowsSdkVersion, null, Logger);
+			return new HoloLensEnvironment(Params, Logger);
+		}
+
 		public override void ValidateTarget(TargetRules Target)
 		{
 			// WindowsTargetRules are reused for HoloLens, so that build modules can keep the model that reuses "windows" configs for most cases
@@ -245,6 +286,12 @@ namespace UnrealBuildTool
 				Target.WindowsPlatform.WindowsSdkVersion = Target.HoloLensPlatform.Win10SDKVersionString;
 			}
 
+			// set the correct architecture
+			if (Target.Architecture.ToLower() == "arm64")
+			{
+				Target.WindowsPlatform.Architecture = WindowsArchitecture.ARM64;
+			}			
+
 			base.ValidateTarget(Target);
 
 			// Windows 10 SDK version
diff --git a/Engine/Source/Programs/UnrealBuildTool/Configuration/ModuleRules.cs b/Engine/Source/Programs/UnrealBuildTool/Configuration/ModuleRules.cs
--- a/Engine/Source/Programs/UnrealBuildTool/Configuration/ModuleRules.cs
+++ b/Engine/Source/Programs/UnrealBuildTool/Configuration/ModuleRules.cs
@@ -1288,20 +1288,18 @@ namespace UnrealBuildTool
 					PrivateDependencyModuleNames.Add("APEX");
 					PublicDefinitions.Add("WITH_APEX=1");
 
-// @MIXEDREALITY_CHANGE : BEGIN - Do not use Apex Cloth for HoloLens.  TODO: can we enable this in the future?
-				if (Target.Platform == UnrealTargetPlatform.HoloLens)
-				{
-					PublicDefinitions.Add("WITH_APEX_CLOTHING=0");
-					PublicDefinitions.Add("WITH_CLOTH_COLLISION_DETECTION=0");
-				}
-				else
-				{
-					PublicDefinitions.Add("WITH_APEX_CLOTHING=1");
-					PublicDefinitions.Add("WITH_CLOTH_COLLISION_DETECTION=1");
-				}
-// @MIXEDREALITY_CHANGE : END
+					if (Target.Platform.ToString() == "HoloLens") //  Do not use Apex Cloth for HoloLens.  TODO: can we enable this in the future?
+					{
+						PublicDefinitions.Add("WITH_APEX_CLOTHING=0");
+						PublicDefinitions.Add("WITH_CLOTH_COLLISION_DETECTION=0");
+					}
+					else
+					{
+						PublicDefinitions.Add("WITH_APEX_CLOTHING=1");
+						PublicDefinitions.Add("WITH_CLOTH_COLLISION_DETECTION=1");
+					}
 
-				PublicDefinitions.Add("WITH_PHYSX_COOKING=1");  // APEX currently relies on cooking even at runtime
+					PublicDefinitions.Add("WITH_PHYSX_COOKING=1");  // APEX currently relies on cooking even at runtime
 
 				}
 				else
@@ -1554,7 +1552,7 @@ namespace UnrealBuildTool
 			string? Architecture = null;
 			string Linkage = string.Empty;
 			string Toolset = string.Empty;
-			if (Target.Platform == UnrealTargetPlatform.Win64 || Target.Platform == UnrealTargetPlatform.HoloLens)
+			if (Target.Platform == UnrealTargetPlatform.Win64)
 			{
 				Platform = "windows";
 				Architecture = Target.WindowsPlatform.Architecture.ToString().ToLowerInvariant();
diff --git a/Engine/Source/Programs/UnrealBuildTool/Configuration/UEBuildPlatform.cs b/Engine/Source/Programs/UnrealBuildTool/Configuration/UEBuildPlatform.cs
--- a/Engine/Source/Programs/UnrealBuildTool/Configuration/UEBuildPlatform.cs
+++ b/Engine/Source/Programs/UnrealBuildTool/Configuration/UEBuildPlatform.cs
@@ -512,8 +512,7 @@ namespace UnrealBuildTool
 			{
 				return ":";
 			}
-			if (BuildHostPlatform.Current.Platform == UnrealTargetPlatform.Win64 || 
-				BuildHostPlatform.Current.Platform == UnrealTargetPlatform.HoloLens)
+			if (BuildHostPlatform.Current.Platform == UnrealTargetPlatform.Win64)
 			{
 				return ";";
 			}
diff --git a/Engine/Source/Programs/UnrealBuildTool/Configuration/UEBuildTarget.cs b/Engine/Source/Programs/UnrealBuildTool/Configuration/UEBuildTarget.cs
--- a/Engine/Source/Programs/UnrealBuildTool/Configuration/UEBuildTarget.cs
+++ b/Engine/Source/Programs/UnrealBuildTool/Configuration/UEBuildTarget.cs
@@ -4130,7 +4130,7 @@ namespace UnrealBuildTool
 				}
 				else
 				{
-					IsCurrentPlatform = Platform.IsInGroup(UnrealPlatformGroup.Windows) || Platform == UnrealTargetPlatform.HoloLens;
+					IsCurrentPlatform = Platform.IsInGroup(UnrealPlatformGroup.Windows);
 				}
 
 				if (IsCurrentPlatform)
diff --git a/Engine/Source/Programs/UnrealBuildTool/Platform/Windows/UEBuildWindows.cs b/Engine/Source/Programs/UnrealBuildTool/Platform/Windows/UEBuildWindows.cs
--- a/Engine/Source/Programs/UnrealBuildTool/Platform/Windows/UEBuildWindows.cs
+++ b/Engine/Source/Programs/UnrealBuildTool/Platform/Windows/UEBuildWindows.cs
@@ -799,11 +799,7 @@ namespace UnrealBuildTool
 		[SupportedOSPlatform("windows")]
 		public override void ValidateTarget(TargetRules Target)
 		{
-			if (Platform == UnrealTargetPlatform.HoloLens && Target.Architecture.ToLower() == "arm64")
-			{
-				Target.WindowsPlatform.Architecture = WindowsArchitecture.ARM64;
-			}
-			else if (Platform == UnrealTargetPlatform.Win64)
+			if (Platform == UnrealTargetPlatform.Win64)
 			{
 				Target.WindowsPlatform.Architecture = WindowsArchitecture.x64;
 			}
diff --git a/Engine/Source/Programs/UnrealBuildTool/Platform/Windows/VCEnvironment.cs b/Engine/Source/Programs/UnrealBuildTool/Platform/Windows/VCEnvironment.cs
--- a/Engine/Source/Programs/UnrealBuildTool/Platform/Windows/VCEnvironment.cs
+++ b/Engine/Source/Programs/UnrealBuildTool/Platform/Windows/VCEnvironment.cs
@@ -282,22 +282,11 @@ namespace UnrealBuildTool
 			// Add the standard Visual C++ library paths
 			if (ToolChain.IsMSVC())
 			{
-				if (Platform == UnrealTargetPlatform.HoloLens)
-				{
-					return DirectoryReference.Combine(ToolChainDir, "lib", ArchFolder, "store");
-				}
-				else
-				{
-					return DirectoryReference.Combine(ToolChainDir, "lib", ArchFolder);
-				}
+				return DirectoryReference.Combine(ToolChainDir, "lib", ArchFolder);
 			}
 			else
 			{
 				DirectoryReference LibsPath = DirectoryReference.Combine(ToolChainDir, "LIB");
-				if (Platform == UnrealTargetPlatform.HoloLens)
-				{
-					LibsPath = DirectoryReference.Combine(LibsPath, "store");
-				}
 
 				if (Architecture == WindowsArchitecture.x64)
 				{
diff --git a/Engine/Source/Programs/UnrealBuildTool/ProjectFiles/VisualStudio/VCProject.cs b/Engine/Source/Programs/UnrealBuildTool/ProjectFiles/VisualStudio/VCProject.cs
--- a/Engine/Source/Programs/UnrealBuildTool/ProjectFiles/VisualStudio/VCProject.cs
+++ b/Engine/Source/Programs/UnrealBuildTool/ProjectFiles/VisualStudio/VCProject.cs
@@ -916,10 +916,6 @@ namespace UnrealBuildTool
 					{
 						SharedIncludeSearchPaths.Append(VCToolChain.GetVCIncludePaths(UnrealTargetPlatform.Win64, GetCompilerForIntellisense(), null, Logger) + ";");
 					}
-					else if (InPlatforms.Contains(UnrealTargetPlatform.HoloLens))
-					{
-						SharedIncludeSearchPaths.Append(VCToolChain.GetVCIncludePaths(UnrealTargetPlatform.HoloLens, GetCompilerForIntellisense(), null, Logger) + ";");
-					}
 				}
 				else
 				{
