fix HoloLens build logging exception.

#rnx

diff --git a/Engine/Platforms/Hololens/Source/Programs/UnrealBuildTool/UEBuildHoloLens.cs b/Engine/Platforms/Hololens/Source/Programs/UnrealBuildTool/UEBuildHoloLens.cs
--- a/Engine/Platforms/Hololens/Source/Programs/UnrealBuildTool/UEBuildHoloLens.cs
+++ b/Engine/Platforms/Hololens/Source/Programs/UnrealBuildTool/UEBuildHoloLens.cs
@@ -310,7 +310,7 @@
 				}
 				else if (Target.HoloLensPlatform.Win10SDKVersion > MaximumSDKVersionTested)
 				{
-					Log.TraceInformationOnce("Your Windows SDK version ({SdkVer}) for HoloLens is newer than the highest tested with this version of UBT ({MaxTestedVer}).  This is probably fine, but if you encounter issues consider using an earlier SDK.", Target.HoloLensPlatform.Win10SDKVersion, MaximumSDKVersionTested);
+					Logger.LogInformation("Your Windows SDK version ({SdkVer}) for HoloLens is newer than the highest tested with this version of UBT ({MaxTestedVer}).  This is probably fine, but if you encounter issues consider using an earlier SDK.", Target.HoloLensPlatform.Win10SDKVersion, MaximumSDKVersionTested);
 				}
 			}
 
