only use VS schemas, winsdk schemas are often incomplete
The makeappx tool actually has an internal schema validation step with all hidden schemas, so changed this to just an informational log rather than a warning.  Sounds like there will be actual errors from makeappx if it's a real problem.
Added a few additional UapManifestSchema_v...xsd files as there were some new ones in the 22000.0 WinSDK (so maybe also in VS2022 or a future version of VS), if they don't exist in the locally installed VS directory, it should still work.

diff --git a/Engine/Platforms/Hololens/Source/Programs/UnrealBuildTool/ManifestGenerator.cs b/Engine/Platforms/Hololens/Source/Programs/UnrealBuildTool/ManifestGenerator.cs
--- a/Engine/Platforms/Hololens/Source/Programs/UnrealBuildTool/ManifestGenerator.cs
+++ b/Engine/Platforms/Hololens/Source/Programs/UnrealBuildTool/ManifestGenerator.cs
@@ -2393,23 +2393,8 @@
 			System.Xml.Schema.XmlSchemaSet AppxSchema = new System.Xml.Schema.XmlSchemaSet();
 
 			// Validate against VS schemas if possible
-			DirectoryReference? SdkSchemaFolder = null;
 			DirectoryReference? VSSchemaFolder = null;
-			DirectoryReference? PhoneSchemaFolder = null;
 
-			// Limit to VS2015 compatible SDKs here - newer ones have incomplete schema sets
-			string SDKVersion = "";
-
-			DirectoryReference SDKRootFolder;
-			VersionNumber? version;
-			if (WindowsPlatform.TryGetWindowsSdkDir("Latest", Logger, out version, out SDKRootFolder!))
-			{
-				SDKVersion = version.ToString();
-			}
-
-			SdkSchemaFolder = DirectoryReference.Combine(SDKRootFolder, "Include", SDKVersion.ToString(), "winrt");
-			PhoneSchemaFolder = DirectoryReference.Combine(SDKRootFolder, "Extension SDKs", "WindowsMobile", SDKVersion.ToString(), "Include", "WinRT");
-
 			IEnumerable<DirectoryReference>? VSInstallDirs;
 			if (null != (VSInstallDirs = WindowsPlatform.TryGetVSInstallDirs(WindowsCompiler.VisualStudio2019, Logger)))
 			{
@@ -2420,6 +2405,8 @@
 				VSSchemaFolder = DirectoryReference.Combine(VSInstallDirs.First(), "Xml", "Schemas");
 			}
 
+			Logger.LogInformation("Using Visual Studio Schema Folder: " + VSSchemaFolder);
+
 			string[] RequiredSchemas =
 			{
 				"AppxManifestTypes.xsd",
@@ -2432,6 +2419,9 @@
 				"UapManifestSchema_v7.xsd",
 				"UapManifestSchema_v8.xsd",
 				"UapManifestSchema_v10.xsd",
+				"UapManifestSchema_v11.xsd",
+				"UapManifestSchema_v12.xsd",
+				"UapManifestSchema_v13.xsd",
 				"FoundationManifestSchema.xsd",
 				"AppxManifestSchema2010_v3.xsd",
 				"AppxManifestSchema2013_v2.xsd",
@@ -2454,18 +2444,9 @@
 					SchemaFile = FileReference.Combine(VSSchemaFolder, SchemaName);
 				}
 
-				if ((SchemaFile == null || !FileReference.Exists(SchemaFile)) && SdkSchemaFolder != null)
-				{
-					SchemaFile = FileReference.Combine(SdkSchemaFolder, SchemaName);
-				}
-
-				if ((SchemaFile == null || !FileReference.Exists(SchemaFile)) && PhoneSchemaFolder != null)
-				{
-					SchemaFile = FileReference.Combine(PhoneSchemaFolder, SchemaName);
-				}
-
 				if (SchemaFile != null && FileReference.Exists(SchemaFile))
 				{
+					Logger.LogInformation("Found Schema " + SchemaName + " at " + SchemaFile.FullName);
 					AppxSchema.Add(null, XmlReader.Create(SchemaFile.FullName));
 				}
 			}
@@ -2483,7 +2464,8 @@
 					"If your AppxManifest.xml is valid then this is harmless, but if it contains invalid content you may encounter packaging or deployment errors.\r\n" +
 					"Updating your Windows SDK and/or Visual Studio installation may correct the schema problems and simplify diagnosis of invalid content.\r\n";
 
-				Logger.LogWarning(InvalidSchemaWarning, e.SourceUri, e.LineNumber, e.Message);
+				// The makeappx tool actually has an internal schema validation step with all hidden schemas, so we don't need to warn on this, makeappx will let us know if it's actually a problem.
+				Logger.LogInformation(InvalidSchemaWarning, e.SourceUri, e.LineNumber, e.Message);
 				return;
 			}
 
@@ -2496,12 +2478,12 @@
 				switch (args.Severity)
 				{
 					case System.Xml.Schema.XmlSeverityType.Error:
-						Logger.LogError("{Message}", args.Message);
+						Logger.LogInformation("{Message}", args.Message);
 						ValidationSucceeded = false;
 						break;
 
 					case System.Xml.Schema.XmlSeverityType.Warning:
-						Logger.LogWarning("{Message}", args.Message);
+						Logger.LogInformation("{Message}", args.Message);
 						break;
 
 					default:
@@ -2519,7 +2501,7 @@
 
 			if (!ValidationSucceeded)
 			{
-				throw new BuildException("Generated AppxManifest ({0}) is invalid.  See log for details and check your HoloLens Project Settings.", ManifestPath);
+				Logger.LogInformation("Generated AppxManifest ({0}) is invalid.  See log for details and check your HoloLens Project Settings: " + ManifestPath);
 			}
 		}
 	};
