Alphabetize and add additional undefs to PostWindowsApi
#rnx

diff --git a/Engine/Platforms/Hololens/Source/Runtime/Core/Public/Microsoft/PostWindowsApiPrivate.h b/Engine/Platforms/Hololens/Source/Runtime/Core/Public/Microsoft/PostWindowsApiPrivate.h
--- a/Engine/Platforms/Hololens/Source/Runtime/Core/Public/Microsoft/PostWindowsApiPrivate.h
+++ b/Engine/Platforms/Hololens/Source/Runtime/Core/Public/Microsoft/PostWindowsApiPrivate.h
@@ -12,37 +12,40 @@
 #undef float
 #undef CDECL
 #undef PF_MAX
-#undef PlaySound
-#undef DrawText
 #undef CaptureStackBackTrace
-#undef MemoryBarrier
-#undef DeleteFile
-#undef MoveFile
 #undef CopyFile
+#undef CreateDesktop
 #undef CreateDirectory
-#undef GetCurrentTime
-#undef SendMessage
-#undef LoadString
-#undef UpdateResource
+#undef CreateFont
+#undef DeleteFile
+#undef DrawText
 #undef FindWindow
-#undef GetObject
+#undef GetClassInfo
+#undef GetClassName
+#undef GetCommandLine
+#undef GetCurrentTime
 #undef GetEnvironmentVariable
-#undef CreateFont
-#undef CreateDesktop
+#undef GetFileAttributes
+#undef GetFreeSpace
 #undef GetMessage
-#undef PostMessage
-#undef GetCommandLine
+#undef GetNextSibling
+#undef GetObject
 #undef GetProp
+#undef GetTempFileName
+#undef IMediaEventSink
+#undef IsMaximized
+#undef IsMinimized
+#undef LoadString
+#undef MemoryBarrier
+#undef MoveFile
+#undef PlaySound
+#undef PostMessage
+#undef ReportEvent
+#undef SendMessage
 #undef SetPort
 #undef SetProp
-#undef GetFileAttributes
-#undef ReportEvent
-#undef GetClassName
-#undef GetClassInfo
+#undef UpdateResource
 #undef Yield
-#undef IMediaEventSink
-#undef GetTempFileName
-#undef GetFreeSpace
 
 // Undefine all the atomics. AllowMicrosoftPlatformAtomics/HideMicrosoftPlatformAtomics temporarily defining these macros.
 #undef InterlockedIncrement
diff --git a/Engine/Source/Runtime/Core/Public/Windows/PostWindowsApi.h b/Engine/Source/Runtime/Core/Public/Windows/PostWindowsApi.h
--- a/Engine/Source/Runtime/Core/Public/Windows/PostWindowsApi.h
+++ b/Engine/Source/Runtime/Core/Public/Windows/PostWindowsApi.h
@@ -23,37 +23,40 @@
 #undef float
 #undef CDECL
 #undef PF_MAX
-#undef PlaySound
-#undef DrawText
 #undef CaptureStackBackTrace
-#undef MemoryBarrier
-#undef DeleteFile
-#undef MoveFile
 #undef CopyFile
+#undef CreateDesktop
 #undef CreateDirectory
-#undef GetCurrentTime
-#undef SendMessage
-#undef LoadString
-#undef UpdateResource
+#undef CreateFont
+#undef DeleteFile
+#undef DrawText
 #undef FindWindow
-#undef GetObject
+#undef GetClassInfo
+#undef GetClassName
+#undef GetCommandLine
+#undef GetCurrentTime
 #undef GetEnvironmentVariable
-#undef CreateFont
-#undef CreateDesktop
+#undef GetFileAttributes
+#undef GetFreeSpace
 #undef GetMessage
-#undef PostMessage
-#undef GetCommandLine
+#undef GetNextSibling
+#undef GetObject
 #undef GetProp
+#undef GetTempFileName
+#undef IMediaEventSink
+#undef IsMaximized
+#undef IsMinimized
+#undef LoadString
+#undef MemoryBarrier
+#undef MoveFile
+#undef PlaySound
+#undef PostMessage
+#undef ReportEvent
+#undef SendMessage
 #undef SetPort
 #undef SetProp
-#undef GetFileAttributes
-#undef ReportEvent
-#undef GetClassName
-#undef GetClassInfo
+#undef UpdateResource
 #undef Yield
-#undef IMediaEventSink
-#undef GetTempFileName
-#undef GetFreeSpace
 
 // Undefine all the atomics. AllowWindowsPlatformAtomics/HideWindowsPlatformAtomics temporarily defining these macros.
 #undef InterlockedIncrement
