UE5 - Remove email addresses

diff --git a/Engine/Platforms/Hololens/Source/Runtime/ApplicationCore/Private/WindowsGamingInputInterface.cpp b/Engine/Platforms/Hololens/Source/Runtime/ApplicationCore/Private/WindowsGamingInputInterface.cpp
--- a/Engine/Platforms/Hololens/Source/Runtime/ApplicationCore/Private/WindowsGamingInputInterface.cpp
+++ b/Engine/Platforms/Hololens/Source/Runtime/ApplicationCore/Private/WindowsGamingInputInterface.cpp
@@ -354,7 +354,7 @@ void WindowsGamingInputInterface::UpdateGamepads()
 				// map left/right stick digital inputs to (top 8 bits)
 
 				// kept h/v key switch due to following change.
-				// @BLACK_TUSK_CHANGE: [REDACTED_EMAIL] - BEGIN - Only send a single direction of left stick analog readings as a key down event at a time.
+				// Only send a single direction of left stick analog readings as a key down event at a time.
 				// - This is to work around the fact that SWidget navigation for the left stick responds to OnKeyDown rather than OnAnalogValue
 				// - To get around that without this change would require a number of changes to the way SWidget handles navigation, 
 				// - and not only that, but any of our UMG screens that implement OnPreviewKeyDown or OnKeyDown themselves to handle left stick navigation manually would need changes as well.
diff --git a/Engine/Platforms/Hololens/Source/Runtime/Core/Public/HoloLensPlatform.h b/Engine/Platforms/Hololens/Source/Runtime/Core/Public/HoloLensPlatform.h
--- a/Engine/Platforms/Hololens/Source/Runtime/Core/Public/HoloLensPlatform.h
+++ b/Engine/Platforms/Hololens/Source/Runtime/Core/Public/HoloLensPlatform.h
@@ -81,9 +81,7 @@ typedef FHoloLensTypes FPlatformTypes;
 [REDACTED_EMAIL]: Fixup once sockets are supported
 #define PLATFORM_SUPPORTS_MESSAGEBUS						1
 
-// @SPLASH_DAMAGE_CHANGE: [REDACTED_EMAIL] - BEGIN
 #define PLATFORM_SUPPORTS_XBOX_LIVE 0
-// @SPLASH_DAMAGE_CHANGE: [REDACTED_EMAIL] - END
 
 #define PLATFORM_HAS_128BIT_ATOMICS							PLATFORM_64BITS
 
