Support for Hololens gauntlet automation launching to device

diff --git a/Engine/Platforms/Hololens/Source/Programs/AutomationTool/Gauntlet/Gauntlet.HoloLensBuildSource.cs b/Engine/Platforms/Hololens/Source/Programs/AutomationTool/Gauntlet/Gauntlet.HoloLensBuildSource.cs
new file mode 100644
--- /dev/null
+++ b/Engine/Platforms/Hololens/Source/Programs/AutomationTool/Gauntlet/Gauntlet.HoloLensBuildSource.cs
@@ -0,0 +1,165 @@
+﻿// Copyright Epic Games, Inc. All Rights Reserved.
+
+using AutomationTool;
+using System;
+using System.Collections.Generic;
+using System.IO;
+using System.Linq;
+//using Tools.DotNETCommon;
+using UnrealBuildTool;
+
+namespace Gauntlet 
+{
+	public class HoloLensBuildSource : StagedBuildSource<StagedBuild>
+	{
+		public override string BuildName { get { return "HoloLensStagedBuild"; } }
+
+		public override UnrealTargetPlatform Platform { get { return UnrealTargetPlatform.HoloLens; } }
+
+		public override string PlatformFolderPrefix { get { return "HoloLens"; } }
+
+		HoloLensBuildSource()
+		{
+
+		}
+	}
+
+	/// <summary>
+	/// HoloLens packaged build from a NSP
+	/// </summary>
+	public class HoloLensPackageBuild : IBuild
+	{
+		public BuildFlags Flags { get; protected set; }
+
+		public UnrealTargetConfiguration Configuration { get; protected set; }
+
+		public UnrealTargetPlatform Platform { get { return UnrealTargetPlatform.HoloLens; } }
+
+		public string ProgramId = null;
+
+		public string SourceAppxBundlePath = null;
+		public string SourceCertPath = null;
+
+		public HoloLensPackageBuild(UnrealTargetConfiguration InConfig, string InAppxBundlePath, BuildFlags InFlags, string InSourceCertPath)
+		{
+			Configuration = InConfig;
+			SourceAppxBundlePath = InAppxBundlePath;
+			Flags = InFlags;
+			SourceCertPath = InSourceCertPath;
+		}
+
+		public bool CanSupportRole(UnrealTargetRole Role)
+		{
+			return Role.IsClient();
+		}
+
+		public static IEnumerable<HoloLensPackageBuild> CreateFromPath(string InProjectName, string InPath)
+		{
+			string BuildPath = InPath;
+
+			List<HoloLensPackageBuild> DiscoveredBuilds = new List<HoloLensPackageBuild>();
+
+			DirectoryInfo Di = new DirectoryInfo(BuildPath);
+
+			// find all appxbundle packages
+			FileInfo[] InstallFiles = Di.GetFiles("*.appxbundle");
+
+			foreach (FileInfo Fi in InstallFiles)
+			{
+				FileInfo[] CertFiles = Di.GetFiles("*.cer");
+
+				var UnrealConfig = UnrealHelpers.GetConfigurationFromExecutableName(InProjectName, Fi.Name);
+				BuildFlags Flags = BuildFlags.Packaged | BuildFlags.CanReplaceCommandLine | BuildFlags.Bulk;
+				string AbsPath = Fi.Directory.FullName;
+				string CertPath = null;
+
+				if (CertFiles.Count() > 0)
+				{
+					CertPath = CertFiles[0].FullName;
+				}
+				else
+				{
+					Log.Warning("Found {0} {1} build at {2}, but no cert file present.", UnrealConfig, ((Flags & BuildFlags.Bulk) == BuildFlags.Bulk) ? "(bulk)" : "(not bulk)", AbsPath);
+				}
+
+				Log.Verbose("Pulling package data from {0}", Fi.FullName);
+
+				HoloLensPackageBuild NewBuild = new HoloLensPackageBuild(UnrealConfig, Fi.FullName, Flags, CertPath);
+
+				DiscoveredBuilds.Add(NewBuild);
+
+				Log.Verbose("Found {0} {1} build at {2}", UnrealConfig, ((Flags & BuildFlags.Bulk) == BuildFlags.Bulk) ? "(bulk)" : "(not bulk)", AbsPath);
+
+			}
+
+			return DiscoveredBuilds;
+		}
+	}
+
+	public class HoloLensPackagedBuildSource : IFolderBuildSource
+	{
+		public string BuildName { get { return "HoloLensPackagedBuildSource"; } }
+
+		public bool CanSupportPlatform(UnrealTargetPlatform InPlatform)
+		{
+			return InPlatform == UnrealTargetPlatform.HoloLens;
+		}
+
+		public string ProjectName { get; protected set; }
+
+		public HoloLensPackagedBuildSource()
+		{
+
+		}
+
+		public List<IBuild> GetBuildsAtPath(string InProjectName, string InPath, int MaxRecursion = 3)
+		{
+			List<DirectoryInfo> AllDirs = new List<DirectoryInfo>();
+
+			// c:\path\to\build
+			DirectoryInfo PathDI = new DirectoryInfo(InPath);
+
+			if (PathDI.Name.IndexOf("HoloLens", StringComparison.OrdinalIgnoreCase) >= 0)
+			{
+				AllDirs.Add(PathDI);
+			}
+
+			// find all directories that begin with HoloLens
+			DirectoryInfo[] HoloLensDirs = PathDI.GetDirectories("HoloLens*", SearchOption.TopDirectoryOnly);
+
+			AllDirs.AddRange(HoloLensDirs);
+
+			List<DirectoryInfo> DirsToRecurse = AllDirs;
+
+			// now get subdirs
+			while (MaxRecursion-- > 0)
+			{
+				List<DirectoryInfo> DiscoveredDirs = new List<DirectoryInfo>();
+
+				DirsToRecurse.ToList().ForEach((D) =>
+				{
+					DiscoveredDirs.AddRange(D.GetDirectories("*", SearchOption.TopDirectoryOnly));
+				});
+
+				AllDirs.AddRange(DiscoveredDirs);
+				DirsToRecurse = DiscoveredDirs;
+			}
+
+			List<IBuild> Builds = new List<IBuild>();
+
+			foreach (DirectoryInfo Di in AllDirs)
+			{
+				IEnumerable<HoloLensPackageBuild> FoundBuilds = HoloLensPackageBuild.CreateFromPath(InProjectName, Di.FullName);
+
+				if (FoundBuilds != null)
+				{
+					Builds.AddRange(FoundBuilds);
+				}
+			}
+
+			return Builds;
+		}
+
+	}
+
+}
diff --git a/Engine/Platforms/Hololens/Source/Programs/AutomationTool/Gauntlet/Gauntlet.HoloLensPGO.cs b/Engine/Platforms/Hololens/Source/Programs/AutomationTool/Gauntlet/Gauntlet.HoloLensPGO.cs
new file mode 100644
--- /dev/null
+++ b/Engine/Platforms/Hololens/Source/Programs/AutomationTool/Gauntlet/Gauntlet.HoloLensPGO.cs
@@ -0,0 +1,59 @@
+// Copyright Epic Games, Inc. All Rights Reserved.
+
+using System;
+using System.IO;
+using System.Text;
+using System.Diagnostics;
+using AutomationTool;
+using UnrealBuildTool;
+
+
+namespace Gauntlet
+{
+	[System.Runtime.Versioning.SupportedOSPlatform("windows")]
+	internal class HoloLensPGOPlatform : IPGOPlatform
+	{
+		string LocalOutputDirectory;
+		string LocalProfDataFile;
+
+		public UnrealTargetPlatform GetPlatform()
+		{
+			return UnrealTargetPlatform.HoloLens;
+		}
+
+		public void GatherResults(string ArtifactPath)
+		{
+			throw new NotImplementedException();
+		}
+
+		public void ApplyConfiguration(PGOConfig Config)
+		{
+			// TODO: Is this sufficient for HoloLens?
+			LocalOutputDirectory = Path.GetFullPath(Config.ProfileOutputDirectory);
+			LocalProfDataFile = Path.Combine(LocalOutputDirectory, "profile.profdata");
+
+			UnrealTestRole ClientRole = Config.RequireRole(UnrealTargetRole.Client);
+			ClientRole.CommandLine += " -pgoprofile";
+		}
+
+		public bool TakeScreenshot(ITargetDevice Device, string ScreenshotDirectory, out string ImageFilename)
+		{
+			ImageFilename = "";
+
+			// Get device for screenshot
+			TargetDeviceHoloLens HoloLensDevice = Device as TargetDeviceHoloLens;
+
+			if (HoloLensDevice == null)
+			{
+				return false;
+			}
+
+			ImageFilename = "temp.bmp";
+			return HoloLensDevice.TakeScreenshot(Path.Combine(ScreenshotDirectory, ImageFilename));
+		}
+
+	}
+}
+
+
+
diff --git a/Engine/Platforms/Hololens/Source/Programs/AutomationTool/Gauntlet/Gauntlet.TargetDeviceHoloLens.cs b/Engine/Platforms/Hololens/Source/Programs/AutomationTool/Gauntlet/Gauntlet.TargetDeviceHoloLens.cs
new file mode 100644
--- /dev/null
+++ b/Engine/Platforms/Hololens/Source/Programs/AutomationTool/Gauntlet/Gauntlet.TargetDeviceHoloLens.cs
@@ -0,0 +1,1309 @@
+﻿// Copyright Epic Games, Inc. All Rights Reserved.
+
+using System;
+using System.Collections.Generic;
+using System.IO;
+using AutomationTool;
+using UnrealBuildTool;
+using System.Text.RegularExpressions;
+using System.Threading;
+using System.Diagnostics;
+using System.Linq;
+//using Tools.DotNETCommon;
+using UnrealBuildBase;
+using EpicGames.Core;
+using Microsoft.Extensions.Logging;
+
+namespace Gauntlet
+{
+	[System.Runtime.Versioning.SupportedOSPlatform("windows")]
+	class HoloLensAppInstance : IAppInstance
+	{
+		protected TargetDeviceHoloLens HoloLensDevice;
+		protected HoloLensAppInstall Install;
+
+		IProcessResult OutputProcess;
+
+		public string CommandLine { get { return "Not Implemented"; } }
+
+		public bool HasExited { get { return OutputProcess.HasExited; } }
+
+		public string StdOut { get { return OutputProcess.Output; } }
+
+		public bool WasKilled { get; protected set; }
+
+		public int ExitCode { get { return OutputProcess.ExitCode; } }
+
+		public HoloLensAppInstance(TargetDeviceHoloLens InDevice, HoloLensAppInstall InInstall, IProcessResult InOutputProcess)
+		{
+			HoloLensDevice = InDevice;
+			Install = InInstall;
+			OutputProcess = InOutputProcess;
+		}
+		
+		public string ArtifactPath
+		{			
+			get
+			{
+				if (HasExited)
+				{
+					SaveArtifacts();
+				}
+
+				return Install.ArtifactPath;
+			}
+		}
+
+		bool bHaveSavedArtifacts;
+		DateTime StartTime = DateTime.UtcNow;
+
+		protected void SaveArtifacts()
+		{
+			if (bHaveSavedArtifacts)
+			{
+				return;
+			}
+
+			bHaveSavedArtifacts = true;
+
+			// TODO: Get HoloLens crash dumps or profiling data?
+		}
+
+		public ITargetDevice Device
+		{
+			get
+			{
+				return HoloLensDevice;
+			}
+		}
+
+		public int WaitForExit()
+		{
+			if (!HasExited)
+			{
+				OutputProcess.WaitForExit();
+			}
+
+			return ExitCode;
+		}
+
+		public void Kill()
+		{
+			if (!HasExited)
+			{
+				WasKilled = true;
+				HoloLensDevice.TerminateRunningApp(OutputProcess.GetProcessName());
+				WaitForExit();
+			}
+		}
+
+		~HoloLensAppInstance()
+		{
+		}
+	}
+
+
+	class HoloLensAppInstall : IAppInstall
+	{
+		public string Name { get; private set; }
+
+		public string ExecutablePath { get; private set; }
+
+		// The local path to our Local HoloLens stage
+		public string LocalBuildPath { get; private set; }
+
+		public UnrealAppConfig AppConfig{ get; private set; }
+
+		public TargetDeviceHoloLens HoloLensDevice { get; private set; }
+
+		public ITargetDevice Device { get { return HoloLensDevice; } }
+
+		public CommandUtils.ERunOptions RunOptions { get; set; }
+
+		public string ArtifactPath;
+
+		public HoloLensAppInstall(UnrealAppConfig InAppConfig, string InExecutablePath, string InBuildPath, TargetDeviceHoloLens InDevice)
+		{
+			AppConfig = InAppConfig;
+			ExecutablePath = InExecutablePath;
+			LocalBuildPath = InBuildPath;
+			Name = AppConfig.Name;
+			HoloLensDevice = InDevice;
+			this.RunOptions = CommandUtils.ERunOptions.NoWaitForExit;
+			ArtifactPath = Path.Combine(LocalBuildPath, AppConfig.ProjectName, @"Saved");
+			try
+			{
+				if (!Directory.Exists(ArtifactPath))
+				{
+					Directory.CreateDirectory(ArtifactPath);
+				}
+
+			}
+			catch (Exception Ex)
+			{
+				throw new AutomationException("Unable to setup temporary HoloLens artifact directory: {0}", Ex.Message);
+			}
+		}
+
+		public IAppInstance Run()
+		{
+			return Device.Run(this);
+		}
+	}
+
+	/// <summary>
+	/// Represents a class capable of creating devices of a specific type
+	/// </summary>
+	[System.Runtime.Versioning.SupportedOSPlatform("windows")]
+	public class HoloLensDeviceFactory : IDeviceFactory
+	{
+		public bool CanSupportPlatform(UnrealTargetPlatform? Platform)
+		{
+			return Platform == UnrealTargetPlatform.HoloLens;
+		}
+
+		public ITargetDevice CreateDevice(string Host, string InCachePath, string InParam = null)
+		{
+			return new TargetDeviceHoloLens(Host, InCachePath);
+		}
+	}
+
+	class HoloLensLauncherCreatedProcess : IProcessResult
+	{
+		Process ExternallyLaunchedProcess;
+		System.Threading.Thread LogTailThread;
+		bool AllowSpew;
+
+		public HoloLensLauncherCreatedProcess(Process InExternallyLaunchedProcess, string LogFile, bool InAllowSpew)
+		{
+			ExternallyLaunchedProcess = InExternallyLaunchedProcess;
+			AllowSpew = InAllowSpew;
+
+			// Can't redirect stdout, so tail the log file (copied from another platform, seems to work for HoloLens too I suppose?)
+			if (AllowSpew)
+			{
+				LogTailThread = new System.Threading.Thread(() => TailLogFile(LogFile));
+				LogTailThread.Start();
+			}
+			else
+			{
+				LogTailThread = null;
+			}
+		}
+
+		~HoloLensLauncherCreatedProcess()
+		{
+			if (ExternallyLaunchedProcess != null)
+			{
+				ExternallyLaunchedProcess.Dispose();
+			}
+		}
+
+		public int ExitCode
+		{
+			get
+			{
+				// Can't access the exit code of a process we didn't start (copied from another platform, seems to work for HoloLens too I suppose?)
+				return 0;
+			}
+
+			set
+			{
+				throw new NotImplementedException();
+			}
+		}
+
+		public bool HasExited
+		{
+			get
+			{
+				// Avoid potential access of the exit code. (copied from another platform, seems to work for HoloLens too I suppose?)
+				return ExternallyLaunchedProcess != null && ExternallyLaunchedProcess.HasExited;
+			}
+		}
+
+		public string Output
+		{
+			get
+			{
+				// TODO: Do we need this for HoloLens?
+				return string.Empty;
+			}
+		}
+
+		public Process ProcessObject
+		{
+			get
+			{
+				return ExternallyLaunchedProcess;
+			}
+		}
+
+		public void DisposeProcess()
+		{
+			ExternallyLaunchedProcess.Dispose();
+			ExternallyLaunchedProcess = null;
+		}
+
+		public string GetProcessName()
+		{
+			return ExternallyLaunchedProcess.ProcessName;
+		}
+
+		public void OnProcessExited()
+		{
+			ProcessManager.RemoveProcess(this);
+			if (LogTailThread != null)
+			{
+				LogTailThread.Join();
+				LogTailThread = null;
+			}
+		}
+
+		public void StdErr(object sender, DataReceivedEventArgs e)
+		{
+			throw new NotImplementedException();
+		}
+
+		public void StdOut(object sender, DataReceivedEventArgs e)
+		{
+			throw new NotImplementedException();
+		}
+
+		public void StopProcess(bool KillDescendants = true)
+		{
+			string ProcessNameForLogging = GetProcessName();
+			try
+			{
+				Process ProcessToKill = ExternallyLaunchedProcess;
+				ExternallyLaunchedProcess = null;
+				if (!ProcessToKill.CloseMainWindow())
+				{
+					CommandUtils.LogWarning("{0} did not respond to close request.  Killing...", ProcessNameForLogging);
+					ProcessToKill.Kill();
+					ProcessToKill.WaitForExit(60000);
+				}
+				if (!ProcessToKill.HasExited)
+				{
+					CommandUtils.LogLog("Process {0} failed to exit.", ProcessNameForLogging);
+				}
+				else
+				{
+					CommandUtils.LogLog("Process {0} successfully exited.", ProcessNameForLogging);
+					OnProcessExited();
+				}
+				ProcessToKill.Close();
+			}
+			catch (Exception Ex)
+			{
+				CommandUtils.LogWarning("Exception while trying to kill process {0}:", ProcessNameForLogging);
+				CommandUtils.LogWarning(LogUtils.FormatException(Ex));
+			}
+		}
+
+		public void WaitForExit()
+		{
+			ExternallyLaunchedProcess.WaitForExit();
+			if (LogTailThread != null)
+			{
+				LogTailThread.Join();
+				LogTailThread = null;
+			}
+		}
+
+		private void TailLogFile(string LogFilePath)
+		{
+			string LogName = GetProcessName();
+			while (!HasExited && !File.Exists(LogFilePath))
+			{
+				System.Threading.Thread.Sleep(1000);
+			}
+
+			if (File.Exists(LogFilePath))
+			{
+				using (StreamReader LogReader = new StreamReader(new FileStream(LogFilePath, FileMode.Open, FileAccess.Read, FileShare.ReadWrite)))
+				{
+					while (!HasExited)
+					{
+						System.Threading.Thread.Sleep(5000);
+						while (!LogReader.EndOfStream)
+						{
+							string LogLine = LogReader.ReadLine();
+							Log.Info("{0} : {1}", LogName, LogLine);
+						}
+					}
+				}
+			}
+		}
+	}
+
+#if !__MonoCS__
+	[System.Runtime.Versioning.SupportedOSPlatform("windows")]
+	class HoloLensDevicePortalCreatedProcess : IProcessResult
+	{
+		object StateLock;
+		Microsoft.Tools.WindowsDevicePortal.DevicePortal Portal;
+		uint ProcId;
+		string PackageName;
+		string FriendlyName;
+		protected string OutputLog;
+
+		public bool ProcessHasExited;
+
+		public Stream GetNewestCrashDumpStream(out string Filename)
+		{
+			var Task = Portal.GetAppCrashDumpListAsync();
+			Task.Wait();
+			var Result = Task.Result;
+			Microsoft.Tools.WindowsDevicePortal.DevicePortal.AppCrashDump NewestDump = null;
+			foreach (var Dump in Result)
+			{
+				if (NewestDump == null || Dump.FileDate > NewestDump.FileDate)
+				{
+					NewestDump = Dump;
+				}
+			}
+
+			if (NewestDump != null)
+			{
+				// System.Net.Http.HttpRequestExcept10n: Cannot write more bytes to the buffer than the configured maximum buffer size: 2147483647.
+				if (NewestDump.FileSizeInBytes > 0 && NewestDump.FileSizeInBytes < 2147483647)
+				{
+					Filename = NewestDump.Filename;
+					var DumpTask = Portal.GetAppCrashDumpAsync(NewestDump);
+					DumpTask.Wait();
+					var DumpResult = DumpTask.Result;
+					// Return first dump for now
+					return DumpResult;
+				}
+				else
+				{
+					var ContentsTask = Portal.GetFolderContentsAsync("LocalAppData", "", PackageName);
+					ContentsTask.Wait();
+					var Folder = ContentsTask.Result;
+
+					Microsoft.Tools.WindowsDevicePortal.DevicePortal.FileOrFolderInformation returnedFileInfo = null;
+
+					foreach (var SubFolderContent in Folder.Contents)
+					{
+						Log.Warning(SubFolderContent.ToString());
+						if (SubFolderContent.IsFolder)
+						{
+							returnedFileInfo = FindFileByNameRecursive(SubFolderContent.SubPath, PackageName, NewestDump.Filename);
+							if (returnedFileInfo != null)
+							{
+								break;
+							}
+
+						}
+
+					}
+
+					if (returnedFileInfo != null)
+					{
+
+						var FileTask = Portal.GetFileAsync("LocalAppData", returnedFileInfo.Name, returnedFileInfo.SubPath, PackageName);
+						FileTask.Wait();
+						var File = FileTask.Result;
+						Filename = returnedFileInfo.Name;
+						return File;
+					}
+				}
+			}
+			Filename = null;
+			return null;
+		}
+
+		public Microsoft.Tools.WindowsDevicePortal.DevicePortal.FileOrFolderInformation FindFileByNameRecursive(string RootPath, string InPackageName, string FilenameToFind)
+		{
+
+			var ContentsTask = Portal.GetFolderContentsAsync("LocalAppData", RootPath, InPackageName);
+			ContentsTask.Wait();
+			var SubFolder = ContentsTask.Result;
+
+			foreach (var SubFolderContent in SubFolder.Contents)
+			{
+				Log.Warning(RootPath);
+				Log.Warning(SubFolderContent.ToString());
+				if (!SubFolderContent.IsFolder)
+				{
+					if (SubFolderContent.Name == FilenameToFind)
+					{
+						return SubFolderContent;
+					}
+				}
+			}
+			
+			return null;
+		}
+
+		public HoloLensDevicePortalCreatedProcess(Microsoft.Tools.WindowsDevicePortal.DevicePortal InPortal, string InPackageName, string InFriendlyName, uint InProcId)
+		{
+			Portal = InPortal;
+			ProcId = InProcId;
+			PackageName = InPackageName;
+			FriendlyName = InFriendlyName;
+			StateLock = new object();
+			Portal.RunningProcessesMessageReceived += Portal_RunningProcessesMessageReceived;
+			Portal.StartListeningForRunningProcessesAsync().Wait();
+
+			// No ETW on Xbox One, so can't collect trace that way (copied from another platform, not sure if it's correct for HoloLens)
+			if (Portal.Platform != Microsoft.Tools.WindowsDevicePortal.DevicePortal.DevicePortalPlatforms.XboxOne)
+			{
+				Guid MicrosoftWindowsDiagnoticsLoggingChannelId = new Guid(0x4bd2826e, 0x54a1, 0x4ba9, 0xbf, 0x63, 0x92, 0xb7, 0x3e, 0xa1, 0xac, 0x4a);
+				Portal.ToggleEtwProviderAsync(MicrosoftWindowsDiagnoticsLoggingChannelId).Wait();
+				Portal.RealtimeEventsMessageReceived += Portal_RealtimeEventsMessageReceived;
+				Portal.StartListeningForEtwEventsAsync().Wait();
+			}
+		}
+
+		//(copied from another platform, not sure if it's correct for HoloLens)
+		private void Portal_RealtimeEventsMessageReceived(Microsoft.Tools.WindowsDevicePortal.DevicePortal sender, Microsoft.Tools.WindowsDevicePortal.WebSocketMessageReceivedEventArgs<Microsoft.Tools.WindowsDevicePortal.DevicePortal.EtwEvents> args)
+		{
+			foreach (var EtwEvent in args.Message.Events)
+			{
+				if (EtwEvent.ContainsKey("ProviderName") && EtwEvent.ContainsKey("StringMessage"))
+				{
+					if (Convert.ToUInt32(EtwEvent["ProcessId"], 10) == ProcId)
+					{
+						Log.Warning("{0} : {1}", FriendlyName, EtwEvent["StringMessage"].Trim('\"'));
+						OutputLog += EtwEvent["StringMessage"].Trim('\"') + "\n";
+					}
+				}
+			}
+		}
+
+		private void Portal_RunningProcessesMessageReceived(Microsoft.Tools.WindowsDevicePortal.DevicePortal sender, Microsoft.Tools.WindowsDevicePortal.WebSocketMessageReceivedEventArgs<Microsoft.Tools.WindowsDevicePortal.DevicePortal.RunningProcesses> args)
+		{
+			foreach (var P in args.Message.Processes)
+			{
+				if (P.ProcessId == ProcId)
+				{
+					OutputLog += args.Message;
+					return;
+				}
+			}
+
+			// TODO: Is this correct?
+			ProcessHasExited = true;
+		}
+
+		public int ExitCode
+		{
+			get
+			{
+				return 0;
+			}
+
+			set
+			{
+				throw new NotImplementedException();
+			}
+		}
+
+		public bool HasExited
+		{
+			get
+			{
+				return ProcessHasExited;
+			}
+		}
+
+		public string Output
+		{
+			get
+			{
+				//Event log method used to populate the log below is not giving us as much info, so just copy the whole file if app has exited
+				if (HasExited)
+				{
+					var FileTask = Portal.GetFileAsync("LocalAppData", "EngineTest.log", "\\LocalState\\EngineTest\\Saved\\Logs", PackageName);
+					FileTask.Wait();
+					var File = FileTask.Result;
+					StreamReader reader = new StreamReader(File);
+					string text = reader.ReadToEnd();
+					return text;
+				}
+				else
+				{
+					if (OutputLog != null)
+					{
+						return OutputLog;
+					}
+				}
+				return OutputLog;
+			}
+		}
+
+		public Process ProcessObject
+		{
+			get
+			{
+				return null;
+			}
+		}
+
+		public void DisposeProcess()
+		{
+		}
+
+		public string GetProcessName()
+		{
+			return PackageName;
+		}
+
+		public void OnProcessExited()
+		{
+			Portal.StopListeningForRunningProcessesAsync().Wait();
+			Portal.StopListeningForEtwEventsAsync().Wait();
+			ProcessHasExited = true;
+		}
+
+		public void StdErr(object sender, DataReceivedEventArgs e)
+		{
+			throw new NotImplementedException();
+		}
+
+		public void StdOut(object sender, DataReceivedEventArgs e)
+		{
+			throw new NotImplementedException();
+		}
+
+		public void StopProcess(bool KillDescendants = true)
+		{
+			Portal.TerminateApplicationAsync(PackageName).Wait();
+		}
+
+		public void WaitForExit()
+		{
+			while (!ProcessHasExited)
+			{
+				System.Threading.Thread.Sleep(1000);
+			}
+		}
+	}
+#endif
+
+	/// <summary>
+	/// HoloLens implementation of a device to run applications
+	/// </summary>
+	[System.Runtime.Versioning.SupportedOSPlatform("windows")]
+	public class TargetDeviceHoloLens : ITargetDevice
+	{
+		public class DeviceInfo
+		{
+			public string Name;
+			public string IPAddress;
+			public bool Default;
+			public bool NeedsRemoval;
+
+			public DeviceInfo(string InHostname, string InIPAddress, bool bIsDefault, bool bNeedsRemoval)
+			{
+				Name = InHostname;
+				IPAddress = InIPAddress;
+				Default = bIsDefault;
+				NeedsRemoval = bNeedsRemoval;
+			}
+		}
+
+		public class StateInfo
+		{
+			public enum ePowerStatus
+			{
+				POWER_STATUS_UNKNOWN = 0,
+				POWER_STATUS_ON = 2,
+				POWER_STATUS_OFF = 4
+			}
+
+			public enum eConnectionState
+			{
+				CONNECTION_UNKNOWN = 0, // TODO: Can we tell with HoloLens whether device is in use?
+				CONNECTION_CONNECTED = 1
+			};
+
+			public ePowerStatus PowerStatus;
+			public eConnectionState ConnectionState;
+
+			public StateInfo()
+			{
+				ConnectionState = eConnectionState.CONNECTION_UNKNOWN;
+				PowerStatus = ePowerStatus.POWER_STATUS_OFF;
+			}
+		}
+
+		protected DeviceInfo StaticDeviceInfo;
+		protected StateInfo CachedStateInfo;		
+
+		//Stop spamming external tool access
+		protected DateTime LastInfoTime = DateTime.MinValue;
+
+		public Dictionary<EIntendedBaseCopyDirectory, string> LocalDirectoryMappings { get; protected set; }
+
+		private static List<string> AcceptThumbprints = new List<string>();
+
+		Microsoft.Tools.WindowsDevicePortal.DevicePortal DevicePortal;
+
+		// We care about UserDir in windows as some of the roles may require files going into user instead of build dir.
+		public void PopulateDirectoryMappings(string BasePath)
+		{
+			LocalDirectoryMappings.Add(EIntendedBaseCopyDirectory.Build, Path.Combine(BasePath, "Build"));
+			LocalDirectoryMappings.Add(EIntendedBaseCopyDirectory.Binaries, Path.Combine(BasePath, "Binaries"));
+			LocalDirectoryMappings.Add(EIntendedBaseCopyDirectory.Config, Path.Combine(BasePath, "Saved", "Config"));
+			LocalDirectoryMappings.Add(EIntendedBaseCopyDirectory.Content, Path.Combine(BasePath, "Content"));
+			LocalDirectoryMappings.Add(EIntendedBaseCopyDirectory.Demos, Path.Combine(BasePath, "Saved", "Demos"));
+			LocalDirectoryMappings.Add(EIntendedBaseCopyDirectory.Profiling, Path.Combine(BasePath, "Saved", "Profiling"));
+			LocalDirectoryMappings.Add(EIntendedBaseCopyDirectory.Saved, Path.Combine(BasePath, "Saved"));
+		}
+
+		public Dictionary<EIntendedBaseCopyDirectory, string> GetPlatformDirectoryMappings()
+		{
+			if (LocalDirectoryMappings.Count == 0)
+			{
+				Log.Warning("Platform directory mappings have not been populated for this platform! This should be done within InstallApplication()");
+			}
+			return LocalDirectoryMappings;
+		}
+
+		public TargetDeviceHoloLens(string Host, string InCachePath, string InParam = null)
+		{
+			string TargetEntryString = ""; // String returned by command tools "<Name>\t<IP>"
+
+			bool bIsDefault = Host.Equals("default", StringComparison.OrdinalIgnoreCase);
+			bool bNeedsRemove = false;
+
+			if (bIsDefault)
+			{
+				TargetEntryString = "Local";
+			}
+			else
+			{
+				// Not supported yet
+				TargetEntryString = "Local";
+			}
+			string Hostname = TargetEntryString;
+
+			// default to IPoverUSB localhost address, and not HTTPS to avoid having to add credentials
+			// TODO: Make this data-driven
+			string IPAddress = "http://127.0.0.1:10080/";
+
+			StaticDeviceInfo = new DeviceInfo(Hostname, IPAddress, bIsDefault, bNeedsRemove);
+			CachedStateInfo = new StateInfo();
+			CachedStateInfo.ConnectionState = StateInfo.eConnectionState.CONNECTION_UNKNOWN;
+			CachedStateInfo.PowerStatus = StateInfo.ePowerStatus.POWER_STATUS_UNKNOWN;
+			LocalDirectoryMappings = new Dictionary<EIntendedBaseCopyDirectory, string>();
+
+			{
+				// fake success, mark device as powered on
+				CachedStateInfo.PowerStatus = StateInfo.ePowerStatus.POWER_STATUS_ON;
+			}
+
+			// Placeholder login info
+			// TODO: make this data-driven
+			string DeviceUsername = "admin";
+			string DevicePassword = "password";
+
+			Microsoft.Tools.WindowsDevicePortal.DefaultDevicePortalConnection DevicePortalConnection = new Microsoft.Tools.WindowsDevicePortal.DefaultDevicePortalConnection(IPAddress, DeviceUsername, DevicePassword);
+			DevicePortal = new Microsoft.Tools.WindowsDevicePortal.DevicePortal(DevicePortalConnection);
+			DevicePortal.UnvalidatedCert += (sender, certificate, chain, sslPolicyErrors) =>
+			{
+				return ShouldAcceptCertificate(new System.Security.Cryptography.X509Certificates.X509Certificate2(certificate), true);
+			};
+			DevicePortal.ConnectionStatus += Portal_ConnectionStatus;
+
+			try
+			{
+				DevicePortal.ConnectAsync().Wait();
+			}
+			catch (AggregateException e)
+			{
+				if (e.InnerException is AutomationException)
+				{
+					throw e.InnerException;
+				}
+				else
+				{
+					throw new AutomationException(ExitCode.Error_LauncherFailed, e.InnerException, e.InnerException.Message);
+				}
+			}
+			catch (Exception e)
+			{
+				throw new AutomationException(ExitCode.Error_LauncherFailed, e, e.Message);
+			}
+		}
+
+		private bool ShouldAcceptCertificate(System.Security.Cryptography.X509Certificates.X509Certificate2 Certificate, bool Unattended)
+		{
+			// TODO: Should anything be verified here?
+			return true;
+		}
+
+#if !__MonoCS__
+		private void Portal_AppInstallStatus(Microsoft.Tools.WindowsDevicePortal.DevicePortal sender, Microsoft.Tools.WindowsDevicePortal.ApplicationInstallStatusEventArgs args)
+		{
+			if (args.Status == Microsoft.Tools.WindowsDevicePortal.ApplicationInstallStatus.Failed)
+			{
+				Log.Error("Error_AppInstallFailed: " + args.Message);
+				throw new AutomationException(ExitCode.Error_AppInstallFailed, args.Message);
+			}
+			else
+			{
+				Log.Info(args.Message);
+			}
+		}
+
+		private void Portal_ConnectionStatus(Microsoft.Tools.WindowsDevicePortal.DevicePortal sender, Microsoft.Tools.WindowsDevicePortal.DeviceConnectionStatusEventArgs args)
+		{
+			if (args.Status == Microsoft.Tools.WindowsDevicePortal.DeviceConnectionStatus.Failed)
+			{
+				throw new AutomationException(args.Message);
+			}
+			else
+			{
+				Log.Info(args.Message);
+				CachedStateInfo.ConnectionState = StateInfo.eConnectionState.CONNECTION_CONNECTED;
+			}
+		}
+#endif
+
+		public CommandUtils.ERunOptions RunOptions { get; set; }
+		public UnrealTargetPlatform? Platform { get { return UnrealTargetPlatform.HoloLens; } }
+
+		public string Name
+		{
+			get
+			{
+				return StaticDeviceInfo.Name;
+			}
+		}
+
+		public bool IsAvailable
+		{
+			get
+			{
+				return IsConnected;
+			}
+		}
+
+		public bool IsConnected
+		{
+			get
+			{
+				return CurrentStateInfo.ConnectionState == StateInfo.eConnectionState.CONNECTION_CONNECTED;
+			}
+		}
+
+		public StateInfo CurrentStateInfo
+		{
+			get
+			{
+				if ((DateTime.Now - LastInfoTime).TotalSeconds > 60.0f)
+				{
+					LastInfoTime = DateTime.Now;
+				}
+
+				return CachedStateInfo;
+			}
+		}
+
+		public bool IsOn { get { return CurrentStateInfo.PowerStatus == StateInfo.ePowerStatus.POWER_STATUS_ON; } }
+
+		public bool PowerOn()
+		{
+			CachedStateInfo.PowerStatus = StateInfo.ePowerStatus.POWER_STATUS_ON;
+
+			return true;
+		}
+		public bool PowerOff()
+		{
+			try
+			{
+#if !__MonoCS__
+				DevicePortal.ShutdownAsync().Wait();
+#endif
+			}
+			catch (Exception Ex)
+			{
+				Log.Warning("Failed to shutdown: " + Ex.Message);
+				return false;
+			}
+
+			return true;
+		}
+
+		public bool Reboot()
+		{
+			try
+			{
+#if !__MonoCS__
+				DevicePortal.RebootAsync().Wait();
+#endif
+			}
+			catch (Exception Ex)
+			{
+				Log.Warning("Failed to reboot: " + Ex.Message);
+				return false;
+			}
+
+			return true;
+		}
+
+		public bool Connect()
+		{
+			// TODO: is this correct?
+			return false;
+		}
+
+		private bool AppLaunched = false;
+		private string CurrentlyRunningAppPackageName = "";
+
+		public bool Disconnect()
+		{
+			if (AppLaunched)
+			{
+				return TerminateRunningApp(CurrentlyRunningAppPackageName);
+			}
+
+			return true;
+		}
+
+		// If a test is being run multiple times, we only want to install the first time
+		static Dictionary<string, string> PackagedInstalls = new Dictionary<string, string>();
+
+		public static string[] GetPathToVCLibsPackages(bool UseDebugCrt, WindowsCompiler Compiler, WindowsArchitecture[] ActualArchitectures)
+		{
+			List<DirectoryReference> SdkRootDirs = new List<DirectoryReference>();
+			WindowsExports.EnumerateSdkRootDirs(SdkRootDirs, EpicGames.Core.Log.Logger);
+			
+			string VCVersionFragment;
+            switch (Compiler)
+			{
+				case WindowsCompiler.VisualStudio2019:
+				case WindowsCompiler.VisualStudio2022:
+				//Compiler version is still 14 for 2017
+				case WindowsCompiler.Default:
+					VCVersionFragment = "14";
+					break;
+
+				default:
+					VCVersionFragment = "Unsupported_VC_Version";
+					break;
+			}
+
+			List<string> Runtimes = new List<string>();
+
+			bool foundVcLibs = false;
+			foreach (DirectoryReference SdkRootDir in SdkRootDirs)
+			{
+				if (!foundVcLibs)
+				{
+					foreach (var Arch in ActualArchitectures)
+					{
+						string ArchitectureFragment = WindowsExports.GetArchitectureSubpath(Arch);
+
+						// For whatever reason, VCLibs are not in Windows Kits/, but Microsoft SDKs/Windows Kits, so go up two directories and start from there.
+						if (Directory.Exists(SdkRootDir.ParentDirectory.ToString()) && Directory.Exists(SdkRootDir.ParentDirectory.ParentDirectory.ToString()))
+						{
+							DirectoryReference Win64SDKsRootDir = SdkRootDir.ParentDirectory.ParentDirectory;
+
+							string RuntimePath = Path.Combine(Win64SDKsRootDir.ToString(),
+								"Microsoft SDKs",
+								"Windows Kits",
+								"10",
+								"ExtensionSDKs",
+								"Microsoft.VCLibs",
+								string.Format("{0}.0", VCVersionFragment),
+								"Appx",
+								UseDebugCrt ? "Debug" : "Retail",
+								ArchitectureFragment,
+								string.Format("Microsoft.VCLibs.{0}.{1}.00.appx", ArchitectureFragment, VCVersionFragment));
+
+							if (File.Exists(RuntimePath))
+							{
+								Runtimes.Add(RuntimePath);
+								Log.Info("found vclibs, path: " + RuntimePath);
+								foundVcLibs = true;
+							}
+							else
+							{
+								Log.Info("no vclib found, path: " + RuntimePath);
+							}
+						}
+						else
+						{
+							Log.Info("Directory structure for VCLibs not present , SDK path: " + SdkRootDir.ToString());
+						}
+					}
+				}
+			}
+
+			if (!foundVcLibs)
+			{
+				Log.Error("no appropriate VCLibs found in any of the SDK locations");
+			}
+
+			return Runtimes.ToArray();
+		}
+
+		private void GetPackageInfo(Dictionary<UnrealTargetPlatform, ConfigHierarchy> GameConfigs, UnrealTargetPlatform PlatformType, out string Name, out string Publisher)
+		{
+			ConfigHierarchy PlatformEngineConfig = null;
+			GameConfigs.TryGetValue(PlatformType, out PlatformEngineConfig);
+
+			if (PlatformEngineConfig == null || !PlatformEngineConfig.GetString("/Script/EngineSettings.GeneralProjectSettings", "ProjectName", out Name))
+			{
+				Name = "DefaultUE4Project";
+			}
+
+			Name = Regex.Replace(Name, "[^-.A-Za-z0-9]", "");
+
+			if (PlatformEngineConfig == null || !PlatformEngineConfig.GetString("/Script/EngineSettings.GeneralProjectSettings", "CompanyDistinguishedName", out Publisher))
+			{
+				Publisher = "CN=NoPublisher";
+			}
+		}
+
+		IAppInstall InstallPackagedApplication(HoloLensPackageBuild PackagedBuild, UnrealAppConfig AppConfig)
+		{
+#if !__MonoCS__
+			{
+				int x = 6;
+				Log.Info(x.ToString());
+				try
+				{
+					DevicePortal.ConnectAsync().Wait();
+
+					{
+						var CrashDumps = DevicePortal.GetAppCrashDumpListAsync().Result;
+						Log.Info("Found " + CrashDumps.Count + " crash dumps");
+						foreach (var Dump in CrashDumps)
+						{
+							Log.Info("Deleting crash dump " + Dump.Filename + " from package " + Dump.PackageFullName + " from date " + Dump.FileDateAsString);
+							DevicePortal.DeleteAppCrashDumpAsync(Dump).Wait();
+						}
+					}
+					
+				}
+					catch (Exception e)
+				{
+					Log.Verbose("Exception deleting crash dumps: " + e.Message);
+				}
+
+				//PackagedBuild.Platform;
+				List <UnrealTargetPlatform> ClientTargetPlatformTypes = new List<UnrealTargetPlatform>();
+				List<UnrealTargetConfiguration> ClientConfigsToBuild = new List<UnrealTargetConfiguration>() { UnrealTargetConfiguration.Development };
+				ClientTargetPlatformTypes.Add(PackagedBuild.Platform);
+				var Properties = ProjectUtils.GetProjectProperties(AppConfig.ProjectFile, ClientTargetPlatformTypes, ClientConfigsToBuild, false);
+				Dictionary<UnrealTargetPlatform, ConfigHierarchy> LoadedGameConfigs = Properties.GameConfigs;
+
+				string PackageName;
+				string Publisher;
+				GetPackageInfo(LoadedGameConfigs, PackagedBuild.Platform, out PackageName, out Publisher);
+				try
+				{
+					DevicePortal.ConnectAsync().Wait();
+
+					string NameToUninstall = PackageName;
+
+					if (NameToUninstall.Equals(""))
+					{
+						// If the project has the project name unset (like EngineTest does at the moment), it will fail to uninstall because on device the package is "DefaultUnrealProject".
+						NameToUninstall = "DefaultUnrealProject";
+					}
+					{ 
+						var list = DevicePortal.GetInstalledAppPackagesAsync().Result;
+						Log.Info(String.Format("Current Name = {0}, Publisher = {1}", NameToUninstall, Publisher)); //code for installation debugging 
+						foreach (var p in list.Packages)
+						{
+							Log.Verbose(String.Format("Found installed Package FamilyName = {0}, FullName = {1}, Publisher = {2}", p.FamilyName, p.FullName, p.Publisher));
+						}
+						Log.Info("starting uninstall of any packages that have a package.FamilyName that begins with " + NameToUninstall + " and that are from publisher "+ Publisher);
+						foreach (var pack in list.Packages.FindAll(package => package.FamilyName.StartsWith(NameToUninstall) && package.Publisher == Publisher))
+						{
+							Log.Info("starting uninstall of " + pack.FullName);
+							Log.Info("pack.FullName " + pack.FullName);
+							Log.Info("pack.Name " + pack.Name);
+							Log.Info("pack.FamilyName " + pack.FamilyName);
+							Log.Info("pack.AppId " + pack.AppId);
+							DevicePortal.UninstallApplicationAsync(pack.FullName).Wait();
+							Log.Info("finished uninstall of " + pack.FullName);
+						}
+					}
+				}
+				catch(Exception)
+				{
+					Log.Verbose("Unable to uninstall {0}", PackageName);
+				}
+
+					try
+				{
+					string PackagePath = PackagedBuild.SourceAppxBundlePath;
+					string CertPath = PackagedBuild.SourceCertPath;
+
+					DevicePortal.ConnectAsync().Wait();
+
+					List<string> Dependencies = new List<string>();
+					bool UseDebugCrt = AppConfig.Configuration == UnrealTargetConfiguration.Debug;
+					WindowsCompiler Compiler = WindowsCompiler.Default;
+
+					WindowsArchitecture[] ActualArchitectures = { WindowsArchitecture.ARM64 };
+					Dependencies.AddRange(GetPathToVCLibsPackages(UseDebugCrt, Compiler, ActualArchitectures));
+
+					DevicePortal.AppInstallStatus += Portal_AppInstallStatus;
+
+					Log.Info("starting install of " + AppConfig.ProjectName + " at path " + PackagePath);
+					DevicePortal.InstallApplicationAsync(AppConfig.ProjectName, PackagePath, Dependencies, CertPath).Wait();
+				}
+				catch (AggregateException e)
+				{
+					Log.Error("Exception installing package " + e.Message + e.InnerException.Message);
+					if (e.InnerException is AutomationException)
+					{
+						throw e.InnerException;
+					}
+					else
+					{
+						throw new AutomationException(ExitCode.Error_AppInstallFailed, e.InnerException, e.InnerException.Message);
+					}
+				}
+				catch (Exception e)
+				{
+					Log.Error("Exception installing package " + e.Message + e.InnerException.Message);
+					throw new AutomationException(ExitCode.Error_AppInstallFailed, e, e.Message);
+				}
+			}
+#endif
+
+			string DestPath = Path.Combine(Gauntlet.Globals.TempDir, "LocalHoloLensStage");
+
+			if (LocalDirectoryMappings.Count == 0)
+			{
+				PopulateDirectoryMappings(Path.Combine(DestPath, AppConfig.ProjectName));
+			}
+
+			string InstalledAppxBundlePath = null;
+			if (PackagedInstalls.TryGetValue(StaticDeviceInfo.IPAddress, out InstalledAppxBundlePath))
+			{
+				if (InstalledAppxBundlePath != PackagedBuild.SourceAppxBundlePath)
+				{
+					throw new AutomationException("Mismatch on cached device AppxBundle path:\n{0} vs\n{1}", InstalledAppxBundlePath, PackagedBuild.SourceAppxBundlePath);
+				}
+
+				// already installed
+				Log.Verbose("HoloLens device using previously installed application : {0}", PackagedBuild.SourceAppxBundlePath);
+				return new HoloLensAppInstall(AppConfig, "", Path.GetDirectoryName(PackagedBuild.SourceAppxBundlePath), this); ;
+			}
+
+			PackagedInstalls.Add(StaticDeviceInfo.IPAddress, PackagedBuild.SourceAppxBundlePath);
+
+			return new HoloLensAppInstall(AppConfig, "", Path.GetDirectoryName(PackagedBuild.SourceAppxBundlePath), this);
+		}
+
+		public IAppInstall InstallApplication(UnrealAppConfig AppConfig)
+		{
+			HoloLensPackageBuild PackagedBuild = AppConfig.Build as HoloLensPackageBuild;
+			if (PackagedBuild == null)
+			{
+				throw new AutomationException("Unsupported build type for HoloLens (Build supplied was not staged or packaged build");
+			}
+
+			List<UnrealTargetPlatform> ClientTargetPlatformTypes = new List<UnrealTargetPlatform>();
+			List<UnrealTargetConfiguration> ClientConfigsToBuild = new List<UnrealTargetConfiguration>() { UnrealTargetConfiguration.Development };
+			ClientTargetPlatformTypes.Add(PackagedBuild.Platform);
+			var Properties = ProjectUtils.GetProjectProperties(AppConfig.ProjectFile, ClientTargetPlatformTypes, ClientConfigsToBuild, false);
+			Dictionary<UnrealTargetPlatform, ConfigHierarchy> LoadedGameConfigs = Properties.GameConfigs;
+
+			string PackageName;
+			string Publisher;
+			GetPackageInfo(LoadedGameConfigs, PackagedBuild.Platform, out PackageName, out Publisher);
+
+			string NameToStop = PackageName;
+
+			if (NameToStop.Equals(""))
+			{
+				// If the project has the project name unset (like EngineTest does at the moment), it will fail to uninstall because on device the package is "DefaultUnrealProject".
+				NameToStop = "DefaultUnrealProject";
+			}
+			// close the package name of what we're trying to install
+			try
+			{
+				DevicePortal.ConnectAsync().Wait();
+				Microsoft.Tools.WindowsDevicePortal.DevicePortal.RunningProcesses Processes = DevicePortal.GetRunningProcessesAsync().Result;
+				foreach (var Proc in Processes.Processes)
+				{
+					if (Proc != null && Proc.PackageFullName != null)
+					{
+						Log.Info("Running process: AppName: " + Proc.AppName + "PackageFullName: " + Proc.PackageFullName);
+						if (Proc.IsRunning && Proc.PackageFullName.StartsWith(NameToStop))
+						{
+							Log.Info("Stopping process: AppName: " + Proc.AppName + "PackageFullName: " + Proc.PackageFullName);
+							DevicePortal.TerminateApplicationAsync(NameToStop).Wait();
+						}
+					}
+				}
+				
+			}
+			catch (System.AggregateException AggEx)
+			{
+				Log.Warning(AggEx.InnerException.Message);
+			}
+			catch (Exception Ex)
+			{
+				Log.Warning(Ex.Message);
+			}
+
+			return InstallPackagedApplication(PackagedBuild, AppConfig);
+		}
+
+		private void GetPackageInfo(string AppxManifestPath, out string Name, out string Publisher)
+		{
+			System.Xml.Linq.XDocument Doc = System.Xml.Linq.XDocument.Load(AppxManifestPath);
+			System.Xml.Linq.XElement Package = Doc.Root;
+			System.Xml.Linq.XElement Identity = Package.Element(System.Xml.Linq.XName.Get("Identity", Package.Name.NamespaceName));
+			Name = Identity.Attribute("Name").Value;
+			Publisher = Identity.Attribute("Publisher").Value;
+		}
+
+
+		// path to launched executable (which may be different than the build executable)
+		public string LaunchExePath;
+
+		public IAppInstance Run(IAppInstall App)
+		{
+			HoloLensAppInstall HoloLensApp = App as HoloLensAppInstall;
+			HoloLensPackageBuild PackagedBuild = HoloLensApp.AppConfig.Build as HoloLensPackageBuild;
+
+			if (HoloLensApp == null)
+			{
+				throw new DeviceException("AppInstance is of incorrect type!");
+			}
+			// Probably can't do anything with this for packaged builds?
+			string test = HoloLensApp.AppConfig.CommandLine;
+
+#if !__MonoCS__
+			string Name;
+			string Publisher;
+			FileReference AppxBundleFileRef = new FileReference(PackagedBuild.SourceAppxBundlePath);
+			DirectoryReference RootEngineDir = AppxBundleFileRef.Directory.ParentDirectory.ParentDirectory.ParentDirectory;
+			string AppxManifestFileRef = RootEngineDir.ToString() + "\\" + HoloLensApp.AppConfig.ProjectName + "\\Binaries\\HoloLens\\AppxManifest_arm64.xml";
+			GetPackageInfo(AppxManifestFileRef, out Name, out Publisher);
+
+			try
+			{
+				DevicePortal.ConnectAsync().Wait();
+				string Aumid = string.Empty;
+				string FullName = string.Empty;
+				var AllAppsTask = DevicePortal.GetInstalledAppPackagesAsync();
+				AllAppsTask.Wait();
+				foreach (var AppPackage in AllAppsTask.Result.Packages)
+				{
+					// App.Name seems to report the localized name.
+					if (AppPackage.FamilyName.StartsWith(Name) && AppPackage.Publisher == Publisher)
+					{
+						Aumid = AppPackage.AppId;
+						FullName = AppPackage.FullName;
+						break;
+					}
+				}
+
+				var LaunchTask = DevicePortal.LaunchApplicationAsync(Aumid, FullName);
+
+				LaunchTask.Wait();
+				IProcessResult OutProcess = new HoloLensDevicePortalCreatedProcess(DevicePortal, FullName, App.Name, LaunchTask.Result);
+
+				ProcessManager.AddProcess(OutProcess);
+	
+				CurrentlyRunningAppPackageName = FullName;
+				return new HoloLensAppInstance(this, HoloLensApp, OutProcess); ;
+			}
+			catch (AggregateException e)
+			{
+				if (e.InnerException is AutomationException)
+				{
+					throw e.InnerException;
+				}
+				else
+				{
+					throw new AutomationException(ExitCode.Error_LauncherFailed, e.InnerException, e.InnerException.Message);
+				}
+			}
+			catch (Exception e)
+			{
+				throw new AutomationException(ExitCode.Error_LauncherFailed, e, e.Message);
+			}
+
+
+#else
+						return null;
+#endif
+		}
+
+		public bool TerminateRunningApp(string PackageName)
+		{
+#if !__MonoCS__
+			DevicePortal.TerminateApplicationAsync(PackageName).Wait();
+#endif
+			// Not sure if we should try an catch any errors, just return true
+			return true;
+		}
+
+		public bool TakeScreenshot(string OutputFileName)
+		{
+			try
+			{
+#if !__MonoCS__
+				//TakeMrcPhotoAsync(bool includeHolograms = true, bool includeColorCamera = true);
+				DevicePortal.TakeMrcPhotoAsync(true, false);
+#endif
+			}
+			catch (Exception Ex)
+			{
+				Log.Warning("Failed to take screenshot: " + Ex.Message);
+				return false;
+			}
+
+			return true;
+		}
+
+
+		public bool Disconnect(bool bForce = false)
+		{
+			// TODO: is this correct?
+			return false;
+		}
+
+		// TODO: Do we need this?
+		#region IDisposable Support
+		private bool disposedValue = false; // To detect redundant calls
+
+		protected virtual void Dispose(bool disposing)
+		{
+			if (!disposedValue)
+			{
+				if (disposing)
+				{
+					// TODO: dispose managed state (managed objects).
+				}
+
+				if (IsConnected)
+				{
+					Disconnect();
+				}
+
+
+
+				disposedValue = true;
+			}
+		}
+
+		// This code added to correctly implement the disposable pattern.
+		public void Dispose()
+		{
+			// Do not change this code. Put cleanup code in Dispose(bool disposing) above.
+			Dispose(true);
+		}
+		#endregion
+	}
+}
+ 
+ 
+ 
+ 
\ No newline at end of file
diff --git a/Engine/Platforms/Hololens/Source/Programs/AutomationTool/HoloLens.Automation.csproj b/Engine/Platforms/Hololens/Source/Programs/AutomationTool/HoloLens.Automation.csproj
--- a/Engine/Platforms/Hololens/Source/Programs/AutomationTool/HoloLens.Automation.csproj
+++ b/Engine/Platforms/Hololens/Source/Programs/AutomationTool/HoloLens.Automation.csproj
@@ -38,6 +38,12 @@
     <UseWindowsForms>true</UseWindowsForms>
   </PropertyGroup>
 
+  <ItemGroup>
+    <Compile Remove="Gauntlet\**" />
+    <Compile Remove="LowLevelTests\**" />
+    <Compile Remove="AutomationUtils\**" />
+  </ItemGroup>
+
   <ItemGroup>
     <Compile Include="..\..\..\..\..\..\Engine\Source\Programs\Shared\MetaData.cs">
       <Link>Properties\MetaData.cs</Link>
@@ -48,6 +54,7 @@
     <ProjectReference Include="..\..\..\..\..\..\Engine\Source\Programs\UnrealBuildTool\UnrealBuildTool.csproj" PrivateAssets="All"><Private>false</Private></ProjectReference>
     <ProjectReference Include="..\..\..\..\..\..\Engine\Source\Programs\AutomationTool\AutomationUtils\AutomationUtils.Automation.csproj" PrivateAssets="All"><Private>false</Private></ProjectReference>
     <ProjectReference Include="..\..\..\..\..\..\Engine\Source\Programs\AutomationTool\Scripts\AutomationScripts.Automation.csproj" PrivateAssets="All"><Private>false</Private></ProjectReference>
+    <ProjectReference Include="..\..\..\..\..\Source\Programs\AutomationTool\Gauntlet\Gauntlet.Automation.csproj" />
   </ItemGroup>
   <ItemGroup>
     <PackageReference Include="Microsoft.CSharp" Version="4.7.0" />
diff --git a/Engine/Platforms/Hololens/Source/Programs/AutomationTool/HoloLensPlatform.Automation.cs b/Engine/Platforms/Hololens/Source/Programs/AutomationTool/HoloLensPlatform.Automation.cs
--- a/Engine/Platforms/Hololens/Source/Programs/AutomationTool/HoloLensPlatform.Automation.cs
+++ b/Engine/Platforms/Hololens/Source/Programs/AutomationTool/HoloLensPlatform.Automation.cs
@@ -1748,79 +1748,9 @@ namespace HoloLens.Automation
 			}
 		}
 
-		private string[] GetPathToVCLibsPackages(bool UseDebugCrt, WindowsCompiler Compiler)
-		{
-			List<DirectoryReference> SdkRootDirs = new List<DirectoryReference>();
-			WindowsExports.EnumerateSdkRootDirs(SdkRootDirs, Log.Logger);
-			
-			string VCVersionFragment;
-            switch (Compiler)
-			{
-				case WindowsCompiler.VisualStudio2019:
-				case WindowsCompiler.VisualStudio2022:
-				//Compiler version is still 14 for 2017
-				case WindowsCompiler.Default:
-					VCVersionFragment = "14";
-					break;
-
-				default:
-					VCVersionFragment = "Unsupported_VC_Version";
-					break;
-			}
-
-			List<string> Runtimes = new List<string>();
-
-			bool foundVcLibs = false;
-			foreach (DirectoryReference SdkRootDir in SdkRootDirs)
-			{
-				if (!foundVcLibs)
-				{
-					foreach (var Arch in ActualArchitectures)
-					{
-						string ArchitectureFragment = WindowsExports.GetArchitectureSubpath(Arch);
-
-						// For whatever reason, VCLibs are not in Windows Kits/, but Microsoft SDKs/Windows Kits, so go up two directories and start from there.
-						if (DirectoryExists(SdkRootDir.ParentDirectory.ToString()) && DirectoryExists(SdkRootDir.ParentDirectory.ParentDirectory.ToString()))
-						{
-							DirectoryReference Win64SDKsRootDir = SdkRootDir.ParentDirectory.ParentDirectory;
-
-							string RuntimePath = Path.Combine(Win64SDKsRootDir.ToString(),
-								"Microsoft SDKs",
-								"Windows Kits",
-								"10",
-								"ExtensionSDKs",
-								"Microsoft.VCLibs",
-								string.Format("{0}.0", VCVersionFragment),
-								"Appx",
-								UseDebugCrt ? "Debug" : "Retail",
-								ArchitectureFragment,
-								string.Format("Microsoft.VCLibs.{0}.{1}.00.appx", ArchitectureFragment, VCVersionFragment));
-
-							if (FileExists(RuntimePath))
-							{
-								Runtimes.Add(RuntimePath);
-								LogLog("found vclibs, path: " + RuntimePath);
-								foundVcLibs = true;
-							}
-							else
-							{
-								LogLog("no vclib found, path: " + RuntimePath);
-							}
-						}
-						else
-						{
-							LogLog("Directory structure for VCLibs not present , SDK path: " + SdkRootDir.ToString());
-						}
-					}
-				}
-			}
-
-			if (!foundVcLibs)
-			{
-				LogError("no appropriate VCLibs found in any of the SDK locations");
-			}
-
-			return Runtimes.ToArray();
+		public string[] GetPathToVCLibsPackages(bool UseDebugCrt, WindowsCompiler Compiler)
+		{
+			return Gauntlet.TargetDeviceHoloLens.GetPathToVCLibsPackages(UseDebugCrt, Compiler, ActualArchitectures);
 		}
 
 		private void GenerateDLCManifestIfNecessary(ProjectParams Params, DeploymentContext SC)
@@ -1844,4 +1774,4 @@ namespace HoloLens.Automation
 
 		private static List<string> AcceptThumbprints = new List<string>();
 	}
-}
+}
\ No newline at end of file
diff --git a/Engine/Source/Programs/AutomationTool/Gauntlet/Gauntlet.Automation.csproj b/Engine/Source/Programs/AutomationTool/Gauntlet/Gauntlet.Automation.csproj
--- a/Engine/Source/Programs/AutomationTool/Gauntlet/Gauntlet.Automation.csproj
+++ b/Engine/Source/Programs/AutomationTool/Gauntlet/Gauntlet.Automation.csproj
@@ -50,6 +50,9 @@
       <HintPath>..\..\..\..\Binaries\ThirdParty\Newtonsoft\Newtonsoft.Json.dll</HintPath>
       <Private>True</Private>
     </Reference>
+    <Reference Include="WindowsDevicePortalWrapper">
+      <HintPath>..\..\..\..\Binaries\ThirdParty\WindowsDevicePortalWrapper\WindowsDevicePortalWrapper.dll</HintPath>
+    </Reference>
   </ItemGroup>
 
   <ItemGroup>
