Add SourceDataOverride to other platform settings

diff --git a/Engine/Platforms/Hololens/Source/Developer/HoloLensPlatformEditor/Private/HoloLensTargetSettings.h b/Engine/Platforms/Hololens/Source/Developer/HoloLensPlatformEditor/Private/HoloLensTargetSettings.h
--- a/Engine/Platforms/Hololens/Source/Developer/HoloLensPlatformEditor/Private/HoloLensTargetSettings.h
+++ b/Engine/Platforms/Hololens/Source/Developer/HoloLensPlatformEditor/Private/HoloLensTargetSettings.h
@@ -146,6 +146,10 @@ public:
 	UPROPERTY(config, EditAnywhere, Category = "Audio")
 	FString SpatializationPlugin;
 
+	/** Which of the currently enabled source data override plugins to use. */
+	UPROPERTY(config, EditAnywhere, Category = "Audio")
+	FString SourceDataOverridePlugin;
+
 	/** Which of the currently enabled reverb plugins to use. */
 	UPROPERTY(config, EditAnywhere, Category = "Audio")
 	FString ReverbPlugin;
diff --git a/Engine/Source/Developer/Linux/LinuxTargetPlatform/Classes/LinuxTargetSettings.h b/Engine/Source/Developer/Linux/LinuxTargetPlatform/Classes/LinuxTargetSettings.h
--- a/Engine/Source/Developer/Linux/LinuxTargetPlatform/Classes/LinuxTargetSettings.h
+++ b/Engine/Source/Developer/Linux/LinuxTargetPlatform/Classes/LinuxTargetSettings.h
@@ -27,6 +27,10 @@ public:
 	UPROPERTY(config, EditAnywhere, Category = "Audio")
 	FString SpatializationPlugin;
 
+	/** Which of the currently enabled source data override plugins to use. */
+	UPROPERTY(config, EditAnywhere, Category = "Audio")
+	FString SourceDataOverridePlugin;
+
 	/** Which of the currently enabled reverb plugins to use. */
 	UPROPERTY(config, EditAnywhere, Category = "Audio")
 	FString ReverbPlugin;
diff --git a/Engine/Source/Developer/Mac/MacTargetPlatform/Classes/MacTargetSettings.h b/Engine/Source/Developer/Mac/MacTargetPlatform/Classes/MacTargetSettings.h
--- a/Engine/Source/Developer/Mac/MacTargetPlatform/Classes/MacTargetSettings.h
+++ b/Engine/Source/Developer/Mac/MacTargetPlatform/Classes/MacTargetSettings.h
@@ -114,6 +114,10 @@ public:
 	UPROPERTY(config, EditAnywhere, Category = "Audio")
 	FString SpatializationPlugin;
 
+	/** Which of the currently enabled source data override plugins to use. */
+	UPROPERTY(config, EditAnywhere, Category = "Audio")
+	FString SourceDataOverridePlugin;
+
 	/** Which of the currently enabled reverb plugins to use. */
 	UPROPERTY(config, EditAnywhere, Category = "Audio")
 	FString ReverbPlugin;
diff --git a/Engine/Source/Runtime/Android/AndroidRuntimeSettings/Classes/AndroidRuntimeSettings.h b/Engine/Source/Runtime/Android/AndroidRuntimeSettings/Classes/AndroidRuntimeSettings.h
--- a/Engine/Source/Runtime/Android/AndroidRuntimeSettings/Classes/AndroidRuntimeSettings.h
+++ b/Engine/Source/Runtime/Android/AndroidRuntimeSettings/Classes/AndroidRuntimeSettings.h
@@ -509,6 +509,10 @@ public:
 	UPROPERTY(config, EditAnywhere, Category = "Audio")
 	FString SpatializationPlugin;
 
+	/** Which of the currently enabled source data override plugins to use. */
+	UPROPERTY(config, EditAnywhere, Category = "Audio")
+	FString SourceDataOverridePlugin;
+
 	/** Which of the currently enabled reverb plugins to use. */
 	UPROPERTY(config, EditAnywhere, Category = "Audio")
 	FString ReverbPlugin;
diff --git a/Engine/Source/Runtime/IOS/IOSRuntimeSettings/Classes/IOSRuntimeSettings.h b/Engine/Source/Runtime/IOS/IOSRuntimeSettings/Classes/IOSRuntimeSettings.h
--- a/Engine/Source/Runtime/IOS/IOSRuntimeSettings/Classes/IOSRuntimeSettings.h
+++ b/Engine/Source/Runtime/IOS/IOSRuntimeSettings/Classes/IOSRuntimeSettings.h
@@ -502,6 +502,10 @@ public:
 	UPROPERTY(config, EditAnywhere, Category = "Audio")
 	FString SpatializationPlugin;
 
+	/** Which of the currently enabled source data override plugins to use. */
+	UPROPERTY(config, EditAnywhere, Category = "Audio")
+	FString SourceDataOverridePlugin;
+
 	/** Which of the currently enabled reverb plugins to use. */
 	UPROPERTY(config, EditAnywhere, Category = "Audio")
 	FString ReverbPlugin;
