UnrealBuildTool: Add support for missing -Compiler flag for several platforms

#rnx

diff --git a/Engine/Platforms/Hololens/Source/Programs/UnrealBuildTool/UEBuildHoloLens.cs b/Engine/Platforms/Hololens/Source/Programs/UnrealBuildTool/UEBuildHoloLens.cs
--- a/Engine/Platforms/Hololens/Source/Programs/UnrealBuildTool/UEBuildHoloLens.cs
+++ b/Engine/Platforms/Hololens/Source/Programs/UnrealBuildTool/UEBuildHoloLens.cs
@@ -35,6 +35,7 @@ namespace UnrealBuildTool
 		[XmlConfigFile(Category = "HoloLensPlatform")]
 		[CommandLine("-2019", Value = nameof(WindowsCompiler.VisualStudio2019))]
 		[CommandLine("-2022", Value = nameof(WindowsCompiler.VisualStudio2022))]
+		[CommandLine("-Compiler=")]
 		public WindowsCompiler Compiler = WindowsCompiler.Default;
 
 		/// <summary>
