UnrealBuildTool: Disallow clang compiler to be used for HoloLens as it doesn't support Windows Runtime extensions

#rnx

diff --git a/Engine/Platforms/Hololens/Source/Programs/UnrealBuildTool/UEBuildHoloLens.cs b/Engine/Platforms/Hololens/Source/Programs/UnrealBuildTool/UEBuildHoloLens.cs
--- a/Engine/Platforms/Hololens/Source/Programs/UnrealBuildTool/UEBuildHoloLens.cs
+++ b/Engine/Platforms/Hololens/Source/Programs/UnrealBuildTool/UEBuildHoloLens.cs
@@ -262,6 +262,12 @@ namespace UnrealBuildTool
 				}
 			}
 
+			if (Target.HoloLensPlatform.Compiler.IsClang())
+			{
+				Logger.LogWarning("Buiding HoloLens with {ClangToolchain} is not supported. The default compiler will be used.", Target.HoloLensPlatform.Compiler);
+				Target.HoloLensPlatform.Compiler = WindowsPlatform.GetDefaultCompiler(Target.ProjectFile, Architecture, Logger);
+			}
+
 			if (!Target.bGenerateProjectFiles)
 			{
 				Log.TraceInformationOnce("Using {0} architecture for deploying to HoloLens device", Architecture);
