Always create a new receipt, don't read the old one.

diff --git a/Engine/Platforms/Hololens/Source/Programs/UnrealBuildTool/HoloLensDeploy.cs b/Engine/Platforms/Hololens/Source/Programs/UnrealBuildTool/HoloLensDeploy.cs
--- a/Engine/Platforms/Hololens/Source/Programs/UnrealBuildTool/HoloLensDeploy.cs
+++ b/Engine/Platforms/Hololens/Source/Programs/UnrealBuildTool/HoloLensDeploy.cs
@@ -490,10 +490,7 @@ namespace UnrealBuildTool
 			// Note: TargetReceipt.Read now expands path variables internally.
 			TargetReceipt? NewReceipt = null;
 			FileReference ReceiptFileName = TargetReceipt.GetDefaultPath(Receipt.ProjectDir != null ? Receipt.ProjectDir : Unreal.EngineDirectory, Receipt.TargetName, Receipt.Platform, Receipt.Configuration, "Multi");
-			if (!TargetReceipt.TryRead(ReceiptFileName, Unreal.EngineDirectory, out NewReceipt))
-			{
-				NewReceipt = new TargetReceipt(Receipt.ProjectFile, Receipt.TargetName, Receipt.TargetType, Receipt.Platform, Receipt.Configuration, Receipt.Version, "Multi", Receipt.IsTestTarget);
-			}
+			NewReceipt = new TargetReceipt(Receipt.ProjectFile, Receipt.TargetName, Receipt.TargetType, Receipt.Platform, Receipt.Configuration, Receipt.Version, "Multi", Receipt.IsTestTarget);
 
 			string SDK = "";
 			var Results = Receipt.AdditionalProperties.Where(x => x.Name == "SDK");
