Changing the agent type used for HoloLens compiling

#rnx

diff --git a/Engine/Platforms/Hololens/Build/InstalledEngineBuild.xml b/Engine/Platforms/Hololens/Build/InstalledEngineBuild.xml
--- a/Engine/Platforms/Hololens/Build/InstalledEngineBuild.xml
+++ b/Engine/Platforms/Hololens/Build/InstalledEngineBuild.xml
@@ -30,7 +30,7 @@
 	</Extend>
 
 	<Extend Name="Platform_Agent_Target_Win64_Licensee">
-		<Agent Name="Target HoloLens" Type="HoloLens">
+		<Agent Name="Target HoloLens" Type="Win64_Licensee">
 			<Node Name="Compile UnrealGame HoloLens" Requires="Compile UnrealHeaderTool Win64" Produces="#UnrealGame HoloLens">
 				<ForEach Name="Configuration" Values="$(GameConfigurations)">
 					<Compile Target="UnrealGame" Platform="HoloLens" Configuration="$(Configuration)" Tag="#UnrealGame HoloLens" Arguments="-precompile -allmodules -nolink -architecture=arm64 $(TargetDebugInfoArg)"/>
