Gauntlet.TargetDeviceHoloLens' TargetDeviceHoloLens constructor needs to use the passed in Host IP, and login credentials for the device need to be data driven from a file in a restricted folder

diff --git a/Engine/Platforms/Hololens/Source/Programs/AutomationTool/Gauntlet/Gauntlet.TargetDeviceHoloLens.cs b/Engine/Platforms/Hololens/Source/Programs/AutomationTool/Gauntlet/Gauntlet.TargetDeviceHoloLens.cs
--- a/Engine/Platforms/Hololens/Source/Programs/AutomationTool/Gauntlet/Gauntlet.TargetDeviceHoloLens.cs
+++ b/Engine/Platforms/Hololens/Source/Programs/AutomationTool/Gauntlet/Gauntlet.TargetDeviceHoloLens.cs
@@ -665,27 +665,22 @@ namespace Gauntlet
 
 		public TargetDeviceHoloLens(string Host, string InCachePath, string InParam = null)
 		{
-			string TargetEntryString = ""; // String returned by command tools "<Name>\t<IP>"
-
 			bool bIsDefault = Host.Equals("default", StringComparison.OrdinalIgnoreCase);
 			bool bNeedsRemove = false;
 
+			// If testing locally, assume it's IpOverUSB (so localhost)
 			if (bIsDefault)
 			{
-				TargetEntryString = "Local";
+				Host = "http://127.0.0.1:10080/";
 			}
+			// Otherwise assume it's a remote URL and (thus needs HTTPS)
 			else
 			{
-				// Not supported yet
-				TargetEntryString = "Local";
+				Host = "https://" + Host;
 			}
-			string Hostname = TargetEntryString;
-
-			// default to IPoverUSB localhost address, and not HTTPS to avoid having to add credentials
-			// TODO: Make this data-driven
-			string IPAddress = "http://127.0.0.1:10080/";
 
-			StaticDeviceInfo = new DeviceInfo(Hostname, IPAddress, bIsDefault, bNeedsRemove);
+			// Not immediately clear how to get a "host name" to pass into Device Info, so just use the host IP for that and IP for now.
+			StaticDeviceInfo = new DeviceInfo(Host, Host, bIsDefault, bNeedsRemove);
 			CachedStateInfo = new StateInfo();
 			CachedStateInfo.ConnectionState = StateInfo.eConnectionState.CONNECTION_UNKNOWN;
 			CachedStateInfo.PowerStatus = StateInfo.ePowerStatus.POWER_STATUS_UNKNOWN;
@@ -694,14 +689,34 @@ namespace Gauntlet
 			{
 				// fake success, mark device as powered on
 				CachedStateInfo.PowerStatus = StateInfo.ePowerStatus.POWER_STATUS_ON;
-			}
+			}	
+
+			string DeviceUsername = "";
+			string DevicePassword = "";
+
+			string DeviceLoginCredentialsFile = CommandUtils.ParseParamValue(Environment.GetCommandLineArgs(), "-HoloLensDeviceCredentialsFileLocation", null);
+
+			if (DeviceLoginCredentialsFile != null && DeviceLoginCredentialsFile.Length > 0)
+			{
+				string[] lines = System.IO.File.ReadAllLines(DeviceLoginCredentialsFile);
+
+				if (lines.Length > 1)
+				{
 
-			// Placeholder login info
-			// TODO: make this data-driven
-			string DeviceUsername = "admin";
-			string DevicePassword = "password";
+					DeviceUsername = lines[0];
+					DevicePassword = lines[1];
+				}
+				else
+				{
+					throw new AutomationException(ExitCode.Error_Arguments,"Device login credential file not found or invalid: " + DeviceLoginCredentialsFile);
+				}
+			}
+			else
+			{
+				throw new AutomationException(ExitCode.Error_Arguments, "Device login credential file not specified. Specify with -HoloLensDeviceCredentialsFileLocation=path.");
+			}
 
-			Microsoft.Tools.WindowsDevicePortal.DefaultDevicePortalConnection DevicePortalConnection = new Microsoft.Tools.WindowsDevicePortal.DefaultDevicePortalConnection(IPAddress, DeviceUsername, DevicePassword);
+			Microsoft.Tools.WindowsDevicePortal.DefaultDevicePortalConnection DevicePortalConnection = new Microsoft.Tools.WindowsDevicePortal.DefaultDevicePortalConnection(Host, DeviceUsername, DevicePassword);
 			DevicePortal = new Microsoft.Tools.WindowsDevicePortal.DevicePortal(DevicePortalConnection);
 			DevicePortal.UnvalidatedCert += (sender, certificate, chain, sslPolicyErrors) =>
 			{
