Fix hololens compilation
#horde 218915

diff --git a/Engine/Platforms/Hololens/Source/Runtime/D3D12RHI/Private/D3D12RHIBasePrivate.h b/Engine/Platforms/Hololens/Source/Runtime/D3D12RHI/Private/D3D12RHIBasePrivate.h
--- a/Engine/Platforms/Hololens/Source/Runtime/D3D12RHI/Private/D3D12RHIBasePrivate.h
+++ b/Engine/Platforms/Hololens/Source/Runtime/D3D12RHI/Private/D3D12RHIBasePrivate.h
@@ -8,12 +8,23 @@
 
 #include "D3D12ThirdParty.h"
 
+// Standard Windows implementation, used to mark the type as 'final'.
+struct FD3D12Payload final : public FD3D12PayloadBase
+{
+	FD3D12Payload(FD3D12Device* const Device, ED3D12QueueType const QueueType)
+		: FD3D12PayloadBase(Device, QueueType)
+	{}
+};
 #define D3D12RHI_RESOURCE_FLAG_ALLOW_INDIRECT_BUFFER D3D12_RESOURCE_FLAG_NONE
 #define D3D12RHI_HEAP_FLAG_ALLOW_INDIRECT_BUFFERS		D3D12_HEAP_FLAG_NONE
 
 #define D3D12RHI_NEEDS_VENDOR_EXTENSIONS     0
 #define D3D12RHI_NEEDS_SHADER_FEATURE_CHECKS 0
 
+// Hololens uses the graphics command list interface for the copy queue.
+typedef ID3D12GraphicsCommandList ID3D12CopyCommandList;
+#define D3D12RHI_PLATFORM_COPY_COMMAND_LIST_TYPE D3D12_COMMAND_LIST_TYPE_COPY
+
 // Heap create not zeroed flag is not available on Hololens so use internal define to disable the feature
 // but make code path shared when it becomes available
 #define FD3D12_HEAP_FLAG_CREATE_NOT_ZEROED D3D12_HEAP_FLAG_NONE
diff --git a/Engine/Platforms/Hololens/Source/Runtime/D3D12RHI/Private/HoloLensD3D12Fence.cpp b/Engine/Platforms/Hololens/Source/Runtime/D3D12RHI/Private/HoloLensD3D12Fence.cpp
deleted file mode 100644
