marc.audy, please merge this change by hand.
More info at https://robomerge.epicgames.net#UE5

{"clientFile":"/src/ROBOMERGE_UE5_UE5_Main/Engine/Platforms/HoloLens/Source/Runtime/D3D12RHI/Private/D3D12RHIBasePrivate.h","fromFile":"//UE5/Release-Engine-Staging/Engine/Platforms/HoloLens/Source/Runtime/D3D12RHI/Private/D3D12RHIBasePrivate.h","startFromRev":2,"endFromRev":3,"resolveType":"content","resolveFlag":"c","contentResolveType":"3waytext"}
Fix Hololens build
 - Hololens uses its own D3D12 platform header, which was missing the new types / classes added by the Dev-ParallelRendering refactor
 - Copied implementation from the Windows header.

diff --git a/Engine/Platforms/Hololens/Source/Runtime/D3D12RHI/Private/D3D12RHIBasePrivate.h b/Engine/Platforms/Hololens/Source/Runtime/D3D12RHI/Private/D3D12RHIBasePrivate.h
--- a/Engine/Platforms/Hololens/Source/Runtime/D3D12RHI/Private/D3D12RHIBasePrivate.h
+++ b/Engine/Platforms/Hololens/Source/Runtime/D3D12RHI/Private/D3D12RHIBasePrivate.h
@@ -6,25 +6,29 @@
 
 #pragma once
 
-#include "D3D12ThirdParty.h"
-
-// Standard Windows implementation, used to mark the type as 'final'.
+// Standard Hololens implementation, used to mark the type as 'final'.
 struct FD3D12Payload final : public FD3D12PayloadBase
 {
 	FD3D12Payload(FD3D12Device* const Device, ED3D12QueueType const QueueType)
 		: FD3D12PayloadBase(Device, QueueType)
 	{}
 };
+
+// Hololens uses the graphics command list interface for the copy queue.
+typedef ID3D12GraphicsCommandList ID3D12CopyCommandList;
+#define D3D12RHI_PLATFORM_COPY_COMMAND_LIST_TYPE D3D12_COMMAND_LIST_TYPE_COPY
+
+// Hololens uses standard D3D12 query heaps for timestamps
+#define D3D12RHI_PLATFORM_USES_TIMESTAMP_QUERIES 1
+
+#include "D3D12ThirdParty.h"
+
 #define D3D12RHI_RESOURCE_FLAG_ALLOW_INDIRECT_BUFFER D3D12_RESOURCE_FLAG_NONE
 #define D3D12RHI_HEAP_FLAG_ALLOW_INDIRECT_BUFFERS		D3D12_HEAP_FLAG_NONE
 
 #define D3D12RHI_NEEDS_VENDOR_EXTENSIONS     0
 #define D3D12RHI_NEEDS_SHADER_FEATURE_CHECKS 0
 
-// Hololens uses the graphics command list interface for the copy queue.
-typedef ID3D12GraphicsCommandList ID3D12CopyCommandList;
-#define D3D12RHI_PLATFORM_COPY_COMMAND_LIST_TYPE D3D12_COMMAND_LIST_TYPE_COPY
-
 // Heap create not zeroed flag is not available on Hololens so use internal define to disable the feature
 // but make code path shared when it becomes available
 #define FD3D12_HEAP_FLAG_CREATE_NOT_ZEROED D3D12_HEAP_FLAG_NONE
