Add missing relaxed atomics to hololens

#rnx
#buildfix

diff --git a/Engine/Platforms/Hololens/Source/Runtime/Core/Public/HoloLensPlatformAtomics.h b/Engine/Platforms/Hololens/Source/Runtime/Core/Public/HoloLensPlatformAtomics.h
--- a/Engine/Platforms/Hololens/Source/Runtime/Core/Public/HoloLensPlatformAtomics.h
+++ b/Engine/Platforms/Hololens/Source/Runtime/Core/Public/HoloLensPlatformAtomics.h
@@ -324,6 +324,30 @@ struct CORE_API FHoloLensAtomics : public FGenericPlatformAtomics
 		return InterlockedCompareExchange((volatile int64*)Src, 0, 0);
 	}
 
+	static FORCEINLINE int8 AtomicRead_Relaxed(volatile const int8* Src)
+	{
+		return *Src;
+	}
+
+	static FORCEINLINE int16 AtomicRead_Relaxed(volatile const int16* Src)
+	{
+		return *Src;
+	}
+
+	static FORCEINLINE int32 AtomicRead_Relaxed(volatile const int32* Src)
+	{
+		return *Src;
+	}
+
+	static FORCEINLINE int64 AtomicRead_Relaxed(volatile const int64* Src)
+	{
+#if PLATFORM_64BITS
+		return *Src;
+#else
+		return InterlockedCompareExchange((volatile int64*)Src, 0, 0);
+#endif
+	}
+
 	static FORCEINLINE void AtomicStore(volatile int8* Src, int8 Val)
 	{
 		InterlockedExchange(Src, Val);
@@ -344,6 +368,30 @@ struct CORE_API FHoloLensAtomics : public FGenericPlatformAtomics
 		InterlockedExchange(Src, Val);
 	}
 
+	static FORCEINLINE void AtomicStore_Relaxed(volatile int8* Src, int8 Val)
+	{
+		*Src = Val;
+	}
+
+	static FORCEINLINE void AtomicStore_Relaxed(volatile int16* Src, int16 Val)
+	{
+		*Src = Val;
+	}
+
+	static FORCEINLINE void AtomicStore_Relaxed(volatile int32* Src, int32 Val)
+	{
+		*Src = Val;
+	}
+
+	static FORCEINLINE void AtomicStore_Relaxed(volatile int64* Src, int64 Val)
+	{
+#if PLATFORM_64BITS
+		* Src = Val;
+#else
+		InterlockedExchange(Src, Val);
+#endif
+	}
+
 	 /**
 	 *	The function compares the Destination value with the Comparand value:
 	 *		- If the Destination value is equal to the Comparand value, the Exchange value is stored in the address specified by Destination, 
