Hololens: Remove deprecated usage of OnControllerAnalog

diff --git a/Engine/Platforms/Hololens/Source/Runtime/ApplicationCore/Private/WindowsGamingInputInterface.cpp b/Engine/Platforms/Hololens/Source/Runtime/ApplicationCore/Private/WindowsGamingInputInterface.cpp
--- a/Engine/Platforms/Hololens/Source/Runtime/ApplicationCore/Private/WindowsGamingInputInterface.cpp
+++ b/Engine/Platforms/Hololens/Source/Runtime/ApplicationCore/Private/WindowsGamingInputInterface.cpp
@@ -239,20 +239,25 @@ void WindowsGamingInputInterface::TerminateGamepadInputs(GamepadMapping& Mapping
 {
 	UE_LOG(GamepadSystem, Log, TEXT("Terminate active gamepad inputs."));
 
+	IPlatformInputDeviceMapper& DeviceMapper = IPlatformInputDeviceMapper::Get();
+	FPlatformUserId PlatformUser = FGenericPlatformMisc::GetPlatformUserForUserIndex(Id);
+	FInputDeviceId DeviceId = INPUTDEVICEID_NONE;
+	DeviceMapper.RemapControllerIdToPlatformUserAndDevice(Id, PlatformUser, DeviceId);
+
 	// reset axes
-	MessageHandler->OnControllerAnalog(FGamepadKeyNames::LeftAnalogX, Id, 0.0f);
-	MessageHandler->OnControllerAnalog(FGamepadKeyNames::LeftAnalogY, Id, 0.0f);
-	MessageHandler->OnControllerAnalog(FGamepadKeyNames::RightAnalogX, Id, 0.0f);
-	MessageHandler->OnControllerAnalog(FGamepadKeyNames::RightAnalogY, Id, 0.0f);
+	MessageHandler->OnControllerAnalog(FGamepadKeyNames::LeftAnalogX, PlatformUser, DeviceId, 0.0f);
+	MessageHandler->OnControllerAnalog(FGamepadKeyNames::LeftAnalogY, PlatformUser, DeviceId, 0.0f);
+	MessageHandler->OnControllerAnalog(FGamepadKeyNames::RightAnalogX, PlatformUser, DeviceId, 0.0f);
+	MessageHandler->OnControllerAnalog(FGamepadKeyNames::RightAnalogY, PlatformUser, DeviceId, 0.0f);
 
 	if (Mapping.LeftTrigger > 0)
 	{
-		MessageHandler->OnControllerAnalog(FGamepadKeyNames::LeftTriggerAnalog, Id, 0.0f);
+		MessageHandler->OnControllerAnalog(FGamepadKeyNames::LeftTriggerAnalog, PlatformUser, DeviceId, 0.0f);
 		Mapping.LeftTrigger = 0.0f;
 	}
 	if (Mapping.RightTrigger > 0)
 	{
-		MessageHandler->OnControllerAnalog(FGamepadKeyNames::RightTriggerAnalog, Id, 0.0f);
+		MessageHandler->OnControllerAnalog(FGamepadKeyNames::RightTriggerAnalog, PlatformUser, DeviceId, 0.0f);
 		Mapping.RightTrigger = 0.0f;
 	}
 
@@ -323,15 +328,20 @@ void WindowsGamingInputInterface::UpdateGamepads()
 				float LeftTrigger = (float)Reading.LeftTrigger;
 				float RightTrigger = (float)Reading.RightTrigger;
 
+				IPlatformInputDeviceMapper& DeviceMapper = IPlatformInputDeviceMapper::Get();
+				FPlatformUserId PlatformUser = FGenericPlatformMisc::GetPlatformUserForUserIndex(i);
+				FInputDeviceId DeviceId = INPUTDEVICEID_NONE;
+				DeviceMapper.RemapControllerIdToPlatformUserAndDevice(i, PlatformUser, DeviceId);
+
 				// map axes
-				MessageHandler->OnControllerAnalog(FGamepadKeyNames::LeftAnalogX, i, LeftThumbX);
-				MessageHandler->OnControllerAnalog(FGamepadKeyNames::LeftAnalogY, i, LeftThumbY);
-				MessageHandler->OnControllerAnalog(FGamepadKeyNames::RightAnalogX, i, RightThumbX);
-				MessageHandler->OnControllerAnalog(FGamepadKeyNames::RightAnalogY, i, RightThumbY);
+				MessageHandler->OnControllerAnalog(FGamepadKeyNames::LeftAnalogX, PlatformUser, DeviceId, LeftThumbX);
+				MessageHandler->OnControllerAnalog(FGamepadKeyNames::LeftAnalogY, PlatformUser, DeviceId, LeftThumbY);
+				MessageHandler->OnControllerAnalog(FGamepadKeyNames::RightAnalogX, PlatformUser, DeviceId, RightThumbX);
+				MessageHandler->OnControllerAnalog(FGamepadKeyNames::RightAnalogY, PlatformUser, DeviceId, RightThumbY);
 
 				// map analog triggers directly
-				MessageHandler->OnControllerAnalog(FGamepadKeyNames::LeftTriggerAnalog, i, LeftTrigger);
-				MessageHandler->OnControllerAnalog(FGamepadKeyNames::RightTriggerAnalog, i, RightTrigger);
+				MessageHandler->OnControllerAnalog(FGamepadKeyNames::LeftTriggerAnalog, PlatformUser, DeviceId, LeftTrigger);
+				MessageHandler->OnControllerAnalog(FGamepadKeyNames::RightTriggerAnalog, PlatformUser, DeviceId, RightTrigger);
 
 				// map buttons (low 15 bits)
 				unsigned int CurrentButtonHeldMask = (((unsigned int)Reading.Buttons) & GamingInputButtonMask);
