Hololens: Remove deprecated usage of OnControllerButtonPressed and OnControllerButtonReleased.

diff --git a/Engine/Platforms/Hololens/Source/Runtime/ApplicationCore/Private/WindowsGamingInputInterface.cpp b/Engine/Platforms/Hololens/Source/Runtime/ApplicationCore/Private/WindowsGamingInputInterface.cpp
--- a/Engine/Platforms/Hololens/Source/Runtime/ApplicationCore/Private/WindowsGamingInputInterface.cpp
+++ b/Engine/Platforms/Hololens/Source/Runtime/ApplicationCore/Private/WindowsGamingInputInterface.cpp
@@ -271,7 +271,7 @@ void WindowsGamingInputInterface::TerminateGamepadInputs(GamepadMapping& Mapping
 		{
 			if (GamepadButtonToUnrealName(ButtonKey, BitMask))
 			{
-				MessageHandler->OnControllerButtonReleased(ButtonKey, Id, false);
+				MessageHandler->OnControllerButtonReleased(ButtonKey, PlatformUser, DeviceId, false);
 			}
 		}
 		BitMask <<= 1;
@@ -427,7 +427,7 @@ void WindowsGamingInputInterface::UpdateGamepads()
 							if (GamepadButtonToUnrealName(ButtonKey, BitMask))
 							{
 								UE_LOG(GamepadSystem, Verbose, TEXT("Gamepad 0x%p (id = %d) - %s Pressed"), (void *)(PadInfo[i].Gamepad), i, *ButtonKey.ToString());
-								MessageHandler->OnControllerButtonPressed(ButtonKey, i, false);
+								MessageHandler->OnControllerButtonPressed(ButtonKey, PlatformUser, DeviceId, false);
 								PadInfo[i].RepeatTime[n] = InitialRepeatDelay;
 							}
 						}
@@ -436,7 +436,7 @@ void WindowsGamingInputInterface::UpdateGamepads()
 							if (GamepadButtonToUnrealName(ButtonKey, BitMask))
 							{
 								UE_LOG(GamepadSystem, Verbose, TEXT("Gamepad 0x%p (id = %d) - %s Released"), (void *)(PadInfo[i].Gamepad), i, *ButtonKey.ToString());
-								MessageHandler->OnControllerButtonReleased(ButtonKey, i, false);
+								MessageHandler->OnControllerButtonReleased(ButtonKey, PlatformUser, DeviceId, false);
 								PadInfo[i].RepeatTime[n] = 0.0f;
 							}
 						}
@@ -452,7 +452,7 @@ void WindowsGamingInputInterface::UpdateGamepads()
 							if (GamepadButtonToUnrealName(ButtonKey, BitMask))
 							{
 								UE_LOG(GamepadSystem, Verbose, TEXT("Gamepad 0x%p (id = %d) - %s Repeated"), (void *)(PadInfo[i].Gamepad), i, *ButtonKey.ToString());
-								MessageHandler->OnControllerButtonPressed(ButtonKey, i, true);
+								MessageHandler->OnControllerButtonPressed(ButtonKey, PlatformUser, DeviceId, true);
 							}
 						}
 					}
