Fix Hololens /permissive- conformance issues

static_cast swtich values, only init one union member
#rnx

diff --git a/Engine/Platforms/Hololens/Source/Runtime/Launch/Private/HoloLensLaunch.cpp b/Engine/Platforms/Hololens/Source/Runtime/Launch/Private/HoloLensLaunch.cpp
--- a/Engine/Platforms/Hololens/Source/Runtime/Launch/Private/HoloLensLaunch.cpp
+++ b/Engine/Platforms/Hololens/Source/Runtime/Launch/Private/HoloLensLaunch.cpp
@@ -284,20 +284,17 @@ uint32 TranslateWinRTKey(Windows::UI::Core::KeyEventArgs^ keyEventArgs)
 	case Windows::System::VirtualKey::RightMenu:	TranslatedKeyCode = VK_RMENU;		break;
 
 	// For the US standard keyboard, only.
-#pragma warning(push)
-#pragma warning(disable:4063)
-	case 0xBA:										TranslatedKeyCode = VK_OEM_1;		break;
-	case 0xBB:										TranslatedKeyCode = VK_OEM_PLUS;	break;
-	case 0xBC:										TranslatedKeyCode = VK_OEM_COMMA;	break;
-	case 0xBD:										TranslatedKeyCode = VK_OEM_MINUS;	break;
-	case 0xBE:										TranslatedKeyCode = VK_OEM_PERIOD;	break;
-	case 0xBF:										TranslatedKeyCode = VK_OEM_2;		break;
-	case 0xC0:										TranslatedKeyCode = VK_OEM_3;		break;
-	case 0xDB:										TranslatedKeyCode = VK_OEM_4;		break;
-	case 0xDC:										TranslatedKeyCode = VK_OEM_5;		break;
-	case 0xDD:										TranslatedKeyCode = VK_OEM_6;		break;
-	case 0xDE:										TranslatedKeyCode = VK_OEM_7;		break;
-#pragma warning(pop)
+	case static_cast<Windows::System::VirtualKey>(0xBA):	TranslatedKeyCode = VK_OEM_1;		break;
+	case static_cast<Windows::System::VirtualKey>(0xBB):	TranslatedKeyCode = VK_OEM_PLUS;	break;
+	case static_cast<Windows::System::VirtualKey>(0xBC):	TranslatedKeyCode = VK_OEM_COMMA;	break;
+	case static_cast<Windows::System::VirtualKey>(0xBD):	TranslatedKeyCode = VK_OEM_MINUS;	break;
+	case static_cast<Windows::System::VirtualKey>(0xBE):	TranslatedKeyCode = VK_OEM_PERIOD;	break;
+	case static_cast<Windows::System::VirtualKey>(0xBF):	TranslatedKeyCode = VK_OEM_2;		break;
+	case static_cast<Windows::System::VirtualKey>(0xC0):	TranslatedKeyCode = VK_OEM_3;		break;
+	case static_cast<Windows::System::VirtualKey>(0xDB):	TranslatedKeyCode = VK_OEM_4;		break;
+	case static_cast<Windows::System::VirtualKey>(0xDC):	TranslatedKeyCode = VK_OEM_5;		break;
+	case static_cast<Windows::System::VirtualKey>(0xDD):	TranslatedKeyCode = VK_OEM_6;		break;
+	case static_cast<Windows::System::VirtualKey>(0xDE):	TranslatedKeyCode = VK_OEM_7;		break;
 
 	default:
 		UE_LOG(LogLaunchHoloLens, Verbose, TEXT("Unrecognized keystroke VirtualKeyCode:%d  ScanCode:%d  Extended:%s  [%s]\n"),
diff --git a/Engine/Source/Runtime/Core/Public/Math/GenericOctree.h b/Engine/Source/Runtime/Core/Public/Math/GenericOctree.h
--- a/Engine/Source/Runtime/Core/Public/Math/GenericOctree.h
+++ b/Engine/Source/Runtime/Core/Public/Math/GenericOctree.h
@@ -316,7 +316,6 @@ public:
 #if PLATFORM_HOLOLENS
 			MaskType() 
 				: v(MakeVectorRegister(0.0f, 0.0f, 0.0f, 0.0f))
-				, i(MakeVectorRegisterInt(0, 0, 0, 0))
 				{}
 #endif
 		} Mask;
