Add missing HTTP_API to many platform http class declarations.

#rnx

diff --git a/Engine/Platforms/Hololens/Source/Runtime/Online/HTTP/Public/HoloLensPlatformHttp.h b/Engine/Platforms/Hololens/Source/Runtime/Online/HTTP/Public/HoloLensPlatformHttp.h
--- a/Engine/Platforms/Hololens/Source/Runtime/Online/HTTP/Public/HoloLensPlatformHttp.h
+++ b/Engine/Platforms/Hololens/Source/Runtime/Online/HTTP/Public/HoloLensPlatformHttp.h
@@ -8,7 +8,7 @@
 /**
 * Platform specific HTTP implementations.
 */
-class FHoloLensHttp : public FGenericPlatformHttp
+class HTTP_API FHoloLensHttp : public FGenericPlatformHttp
 {
 public:
 
diff --git a/Engine/Source/Runtime/Online/HTTP/Public/Android/AndroidPlatformHttp.h b/Engine/Source/Runtime/Online/HTTP/Public/Android/AndroidPlatformHttp.h
--- a/Engine/Source/Runtime/Online/HTTP/Public/Android/AndroidPlatformHttp.h
+++ b/Engine/Source/Runtime/Online/HTTP/Public/Android/AndroidPlatformHttp.h
@@ -6,7 +6,7 @@
 /**
  * Platform specific Http implementations
  */
-class FAndroidPlatformHttp : public FGenericPlatformHttp
+class HTTP_API FAndroidPlatformHttp : public FGenericPlatformHttp
 {
 public:
 
diff --git a/Engine/Source/Runtime/Online/HTTP/Public/Apple/ApplePlatformHttp.h b/Engine/Source/Runtime/Online/HTTP/Public/Apple/ApplePlatformHttp.h
--- a/Engine/Source/Runtime/Online/HTTP/Public/Apple/ApplePlatformHttp.h
+++ b/Engine/Source/Runtime/Online/HTTP/Public/Apple/ApplePlatformHttp.h
@@ -6,7 +6,7 @@
 /**
  * Platform specific Http implementations
  */
-class FApplePlatformHttp : public FGenericPlatformHttp
+class HTTP_API FApplePlatformHttp : public FGenericPlatformHttp
 {
 public:
 
diff --git a/Engine/Source/Runtime/Online/HTTP/Public/Unix/UnixPlatformHttp.h b/Engine/Source/Runtime/Online/HTTP/Public/Unix/UnixPlatformHttp.h
--- a/Engine/Source/Runtime/Online/HTTP/Public/Unix/UnixPlatformHttp.h
+++ b/Engine/Source/Runtime/Online/HTTP/Public/Unix/UnixPlatformHttp.h
@@ -11,7 +11,7 @@ class IHttpRequest;
 /**
  * Platform specific Http implementations
  */
-class FUnixPlatformHttp : public FGenericPlatformHttp
+class HTTP_API FUnixPlatformHttp : public FGenericPlatformHttp
 {
 public:
 
