Exclude HoloLens.Automation.json when not building HoloLens as part of an Installed Engine Build

#rnx

diff --git a/Engine/Platforms/Hololens/Build/InstalledEngineBuild.xml b/Engine/Platforms/Hololens/Build/InstalledEngineBuild.xml
--- a/Engine/Platforms/Hololens/Build/InstalledEngineBuild.xml
+++ b/Engine/Platforms/Hololens/Build/InstalledEngineBuild.xml
@@ -95,6 +95,12 @@
 			<Property Name="CopyInstalledExceptions" Value="$(CopyInstalledExceptions);$(CopyHoloLensExceptions)"/>
 
 		</Do>
+		<Do If="!$(WithHoloLens)">
+			<Property Name="CopyHoloLensCsToolsExceptions">
+				Engine\Saved\CsTools\Engine\Intermediate\ScriptModules\HoloLens.Automation.json
+			</Property>
+			<Property Name="CopyWin64CsToolsExceptions" Value="$(CopyWin64CsToolsExceptions);$(CopyHoloLensCsToolsExceptions)"/>
+		</Do>
 	</Extend>
 
 </BuildGraph>
