Potential fix for warnings in Windows Kits headers during static analysis.

diff --git a/Engine/Platforms/Hololens/Source/Developer/HoloLensTargetPlatform/Private/HoloLensDeviceDetector.cpp b/Engine/Platforms/Hololens/Source/Developer/HoloLensTargetPlatform/Private/HoloLensDeviceDetector.cpp
--- a/Engine/Platforms/Hololens/Source/Developer/HoloLensTargetPlatform/Private/HoloLensDeviceDetector.cpp
+++ b/Engine/Platforms/Hololens/Source/Developer/HoloLensTargetPlatform/Private/HoloLensDeviceDetector.cpp
@@ -10,17 +10,14 @@
 #include "Dom/JsonObject.h"
 #include "HoloLensTargetPlatform.h"
 
-#pragma warning(push)
-#pragma warning(disable:4265 4459 6101)
-
 #include "Serialization/JsonSerializer.h"
 
 #include "Windows/WindowsHWrapper.h"
 #include "Windows/AllowWindowsPlatformTypes.h"
 #include "Windows/AllowWindowsPlatformAtomics.h"
 
+THIRD_PARTY_INCLUDES_START
 #define sealed final 
-
 #include <windows.h>
 #include <wrl/client.h>
 #include <wrl/event.h>
@@ -30,7 +27,7 @@
 #include <Windows.foundation.h>
 #include <Windows.foundation.collections.h>
 #undef sealed  
-#pragma warning(pop)
+THIRD_PARTY_INCLUDES_END
 
 #include <vector>
 #include <memory>
diff --git a/Engine/Platforms/Hololens/Source/Developer/HoloLensTargetPlatform/Private/HoloLensTargetDevice.cpp b/Engine/Platforms/Hololens/Source/Developer/HoloLensTargetPlatform/Private/HoloLensTargetDevice.cpp
--- a/Engine/Platforms/Hololens/Source/Developer/HoloLensTargetPlatform/Private/HoloLensTargetDevice.cpp
+++ b/Engine/Platforms/Hololens/Source/Developer/HoloLensTargetPlatform/Private/HoloLensTargetDevice.cpp
@@ -18,12 +18,10 @@
 #include "Serialization/JsonSerializer.h"
 #include "Misc/ScopeLock.h"
 
-#pragma warning(push)
-#pragma warning(disable:4265 4459)
-
+THIRD_PARTY_INCLUDES_START
 #include <wrl/client.h>
 #include <wrl/wrappers/corewrappers.h>
-#pragma warning(pop)
+THIRD_PARTY_INCLUDES_END
 
 #include <shlwapi.h>
 
