Hololens: do not log username
-Commented out logging of username to align with other platforms on privacy.

diff --git a/Engine/Platforms/Hololens/Source/Runtime/Core/Private/HoloLensPlatformMisc.cpp b/Engine/Platforms/Hololens/Source/Runtime/Core/Private/HoloLensPlatformMisc.cpp
--- a/Engine/Platforms/Hololens/Source/Runtime/Core/Private/HoloLensPlatformMisc.cpp
+++ b/Engine/Platforms/Hololens/Source/Runtime/Core/Private/HoloLensPlatformMisc.cpp
@@ -92,7 +92,7 @@ void FHoloLensMisc::PlatformInit()
 {
 	// Identity.
 	UE_LOG(LogInit, Log, TEXT("Computer: %s"), FPlatformProcess::ComputerName());
-	UE_LOG(LogInit, Log, TEXT("User: %s"), FPlatformProcess::UserName());
+	//UE_LOG(LogInit, Log, TEXT("User: %s"), FPlatformProcess::UserName()); // Do not log user name due to privacy concerns.
 
 	// Timer resolution.
 	UE_LOG(LogInit, Log, TEXT("High frequency timer resolution =%f MHz"), 0.000001 / FPlatformTime::GetSecondsPerCycle());
