UnrealBuildTool: HoloLens doesn't support c++20

#rnx

diff --git a/Engine/Platforms/Hololens/Source/Programs/UnrealBuildTool/UEBuildHoloLens.cs b/Engine/Platforms/Hololens/Source/Programs/UnrealBuildTool/UEBuildHoloLens.cs
--- a/Engine/Platforms/Hololens/Source/Programs/UnrealBuildTool/UEBuildHoloLens.cs
+++ b/Engine/Platforms/Hololens/Source/Programs/UnrealBuildTool/UEBuildHoloLens.cs
@@ -243,6 +243,12 @@ namespace UnrealBuildTool
 
 			ConfigCache.ReadSettings(IniDirRef, Platform, Target.HoloLensPlatform);
 
+			// cl : Command line error D8016 : '/std:c++20' and '/ZW' command-line options are incompatible
+			if (Target.CppStandard > CppStandardVersion.Cpp17)
+			{
+				Target.CppStandard = CppStandardVersion.Cpp17;
+			}
+
 			WindowsArchitecture Architecture = WindowsArchitecture.x64;
 			if( !String.IsNullOrEmpty(Target.Architecture))
 			{
