PR #8297: HoloLens: Implement ComputerName() and GetCurrentProcessId() functions

diff --git a/Engine/Platforms/Hololens/Source/Runtime/Core/Private/HoloLensPlatformProcess.cpp b/Engine/Platforms/Hololens/Source/Runtime/Core/Private/HoloLensPlatformProcess.cpp
--- a/Engine/Platforms/Hololens/Source/Runtime/Core/Private/HoloLensPlatformProcess.cpp
+++ b/Engine/Platforms/Hololens/Source/Runtime/Core/Private/HoloLensPlatformProcess.cpp
@@ -479,4 +479,20 @@ void FHoloLensProcess::LaunchURL(const TCHAR* URL, const TCHAR* Params, FString*
 #endif
 }
 
+const TCHAR* FHoloLensProcess::ComputerName()
+{
+	static TCHAR Result[256] = TEXT("");
+	if (!Result[0])
+	{
+		uint32 Size = UE_ARRAY_COUNT(Result);
+		GetComputerName(Result, (::LPDWORD)&Size);
+	}
+	return Result;
+}
+
+uint32 FHoloLensProcess::GetCurrentProcessId()
+{
+	return ::GetCurrentProcessId();
+}
+
 #include "Microsoft/HideMicrosoftPlatformTypes.h"
diff --git a/Engine/Platforms/Hololens/Source/Runtime/Core/Public/HoloLensPlatformProcess.h b/Engine/Platforms/Hololens/Source/Runtime/Core/Public/HoloLensPlatformProcess.h
--- a/Engine/Platforms/Hololens/Source/Runtime/Core/Public/HoloLensPlatformProcess.h
+++ b/Engine/Platforms/Hololens/Source/Runtime/Core/Public/HoloLensPlatformProcess.h
@@ -43,6 +43,7 @@ struct CORE_API FHoloLensProcess : public FGenericPlatformProcess
 	static void PopDllDirectory(const TCHAR* Directory);
 	static void SetCurrentWorkingDirectoryToBaseDir();
 	static FString GetCurrentWorkingDirectory();
+	static uint32 GetCurrentProcessId();
 	/** Content saved to the game or engine directories should be rerouted to user directories instead **/
 	static bool ShouldSaveToUserDir();
 
@@ -51,6 +52,7 @@ struct CORE_API FHoloLensProcess : public FGenericPlatformProcess
 	static const TCHAR* UserSettingsDir();
 	static const TCHAR* UserTempDir();
 	static const TCHAR* ApplicationSettingsDir();
+	static const TCHAR* ComputerName();
 
 	static const TCHAR* GetLocalAppDataLowLevelPath();
 	static const TCHAR* GetTempAppDataLowLevelPath();
