added missing function implementation for HoloLensWebSocket for submission 23727568

diff --git a/Engine/Platforms/Hololens/Source/Runtime/Online/WebSockets/Private/HoloLensWebSocket.cpp b/Engine/Platforms/Hololens/Source/Runtime/Online/WebSockets/Private/HoloLensWebSocket.cpp
--- a/Engine/Platforms/Hololens/Source/Runtime/Online/WebSockets/Private/HoloLensWebSocket.cpp
+++ b/Engine/Platforms/Hololens/Source/Runtime/Online/WebSockets/Private/HoloLensWebSocket.cpp
@@ -164,6 +164,11 @@ void FHoloLensWebSocket::Send(const void* Utf8Data, SIZE_T Size, bool bIsBinary)
 	}
 }
 
+void FHoloLensWebSocket::SetTextMessageMemoryLimit(uint64 TextMessageMemoryLimit) 
+{
+	UE_LOG(LogWebSockets, Warning, TEXT("SetTextMessageMemoryLimit not implemented for FHoloLensWebSocket."));
+}
+
 bool FHoloLensWebSocket::Tick(float DeltaTime)
 {
 	QUICK_SCOPE_CYCLE_COUNTER(STAT_FHoloLensWebSocket_Tick);
diff --git a/Engine/Platforms/Hololens/Source/Runtime/Online/WebSockets/Private/HoloLensWebSocket.h b/Engine/Platforms/Hololens/Source/Runtime/Online/WebSockets/Private/HoloLensWebSocket.h
--- a/Engine/Platforms/Hololens/Source/Runtime/Online/WebSockets/Private/HoloLensWebSocket.h
+++ b/Engine/Platforms/Hololens/Source/Runtime/Online/WebSockets/Private/HoloLensWebSocket.h
@@ -23,11 +23,13 @@ public:
 	virtual bool IsConnected() override;
 	virtual void Send(const FString& Data) override;
 	virtual void Send(const void* Utf8Data, SIZE_T Size, bool bIsBinary) override;
+	virtual void SetTextMessageMemoryLimit(uint64 TextMessageMemoryLimit) override;
 
 	virtual FWebSocketConnectedEvent& OnConnected() override				{ return ConnectedEvent; }
 	virtual FWebSocketConnectionErrorEvent& OnConnectionError() override	{ return ConnectionErrorEvent; }
 	virtual FWebSocketClosedEvent& OnClosed() override						{ return ClosedEvent; }
 	virtual FWebSocketMessageEvent& OnMessage() override					{ return MessageEvent; }
+	virtual FWebSocketBinaryMessageEvent& OnBinaryMessage() override { return BinaryMessageEvent; }
 	virtual FWebSocketRawMessageEvent& OnRawMessage() override				{ return RawMessageEvent; }
 	virtual FWebSocketMessageSentEvent& OnMessageSent() override { return OnMessageSentEvent; }
 
@@ -50,6 +52,7 @@ private:
 	FWebSocketConnectionErrorEvent ConnectionErrorEvent;
 	FWebSocketClosedEvent ClosedEvent;
 	FWebSocketMessageEvent MessageEvent;
+	FWebSocketBinaryMessageEvent BinaryMessageEvent;
 	FWebSocketRawMessageEvent RawMessageEvent;
 	FWebSocketMessageSentEvent OnMessageSentEvent;
 
