- Set default HoloLens min.maxtested versions in BaseHoloLensEngine.ini

diff --git a/Engine/Platforms/Hololens/Config/BaseHoloLensEngine.ini b/Engine/Platforms/Hololens/Config/BaseHoloLensEngine.ini
new file mode 100644
--- /dev/null
+++ b/Engine/Platforms/Hololens/Config/BaseHoloLensEngine.ini
@@ -0,0 +1,3 @@
+[/Script/HoloLensPlatformEditor.HoloLensTargetSettings]
+MinimumPlatformVersion=10.0.18362.0
+MaximumPlatformVersionTested=10.0.19041.0
