- Changed HoloLens to not error when generating projectfiles with HL platform present

diff --git a/Engine/Platforms/Hololens/Source/Programs/UnrealBuildTool/UEBuildHoloLens.cs b/Engine/Platforms/Hololens/Source/Programs/UnrealBuildTool/UEBuildHoloLens.cs
--- a/Engine/Platforms/Hololens/Source/Programs/UnrealBuildTool/UEBuildHoloLens.cs
+++ b/Engine/Platforms/Hololens/Source/Programs/UnrealBuildTool/UEBuildHoloLens.cs
@@ -11,6 +11,7 @@ using EpicGames.Core;
 using System.Text.RegularExpressions;
 using System.Runtime.Versioning;
 using Microsoft.Extensions.Logging;
+using System.Runtime.InteropServices;
 
 namespace UnrealBuildTool
 {
@@ -204,32 +205,55 @@ namespace UnrealBuildTool
 
 	class HoloLensArchitectureConfig : UnrealArchitectureConfig
 	{
-		public HoloLensArchitectureConfig()
+		private ILogger Logger;
+		public HoloLensArchitectureConfig(ILogger Logger)
 			: base(UnrealArchitectureMode.OneTargetPerArchitecture, new[] { UnrealArch.X64, UnrealArch.Arm64 })
 		{
-
+			this.Logger = Logger;
 		}
 
+		static bool bShownLog = false;
 		public override UnrealArchitectures ActiveArchitectures(FileReference? ProjectFile, string? TargetName)
 		{
 			var ArchList = new List<UnrealArch>();
 			bool bBuildForEmulation;
 			bool bBuildForDevice;
 
+			bool bHasArmToolchain = MicrosoftPlatformSDK.HasValidCompiler(WindowsCompiler.VisualStudio2022, UnrealArch.Arm64, Logger) || MicrosoftPlatformSDK.HasValidCompiler(WindowsCompiler.VisualStudio2019, UnrealArch.Arm64, Logger);
+			bool bHasX64Toolchain = MicrosoftPlatformSDK.HasValidCompiler(WindowsCompiler.VisualStudio2022, UnrealArch.X64, Logger) || MicrosoftPlatformSDK.HasValidCompiler(WindowsCompiler.VisualStudio2019, UnrealArch.X64, Logger);
+	
 			ConfigHierarchy Ini = ConfigCache.ReadHierarchy(ConfigHierarchyType.Engine, ProjectFile?.Directory, UnrealTargetPlatform.HoloLens);
 			if (Ini.GetBool("/Script/HoloLensPlatformEditor.HoloLensTargetSettings", "bBuildForEmulation", out bBuildForEmulation) && bBuildForEmulation)
 			{
-				ArchList.Add(UnrealArch.X64);
+				if (bHasX64Toolchain)
+				{
+					ArchList.Add(UnrealArch.X64);
+				}
 			}
 			if (Ini.GetBool("/Script/HoloLensPlatformEditor.HoloLensTargetSettings", "bBuildForDevice", out bBuildForDevice) && bBuildForDevice)
 			{
-				ArchList.Add(UnrealArch.Arm64);
+				if (bHasArmToolchain)
+				{
+					ArchList.Add(UnrealArch.Arm64);
+				}
 			}
 
 			// if neither were selected, build for device
 			if (ArchList.Count() == 0)
 			{
-				ArchList.Add(UnrealArch.Arm64);
+				if (bHasArmToolchain)
+				{
+					ArchList.Add(UnrealArch.Arm64);
+				}
+				else
+				{
+					if (!bShownLog)
+					{
+						Logger.LogInformation("Arm64 was requested for HoloLens, but the Arm64 toolchain is not installed, defaulting to X64 build");
+						bShownLog = true;
+					}
+					ArchList.Add(UnrealArch.X64);
+				}
 			}
 
 			return new UnrealArchitectures(ArchList);
@@ -243,7 +267,7 @@ namespace UnrealBuildTool
 		public static readonly Version MaximumSDKVersionTested = new Version(10, 0, 18362, int.MaxValue);
 
 		public HoloLensPlatform(MicrosoftPlatformSDK InSDK, ILogger InLogger) 
-			: base(UnrealTargetPlatform.HoloLens, InSDK, new HoloLensArchitectureConfig(), InLogger)
+			: base(UnrealTargetPlatform.HoloLens, InSDK, new HoloLensArchitectureConfig(InLogger), InLogger)
 		{
 		}
 
