AutomationTool: Remove stale -iwyu UBT arg, whatvever it used to do  is no longer available and this flag now runs the IWYU toolchain

#rnx

diff --git a/Engine/Platforms/Hololens/Source/Programs/AutomationTool/Scripts/BuildPluginCommand_HoloLens.Automation.cs b/Engine/Platforms/Hololens/Source/Programs/AutomationTool/Scripts/BuildPluginCommand_HoloLens.Automation.cs
--- a/Engine/Platforms/Hololens/Source/Programs/AutomationTool/Scripts/BuildPluginCommand_HoloLens.Automation.cs
+++ b/Engine/Platforms/Hololens/Source/Programs/AutomationTool/Scripts/BuildPluginCommand_HoloLens.Automation.cs
@@ -24,7 +24,7 @@ public class BuildPlugin_HoloLens : BuildPlugin.TargetPlatform
 		{
 			FileReference ManifestFileName = FileReference.Combine(HostProjectFile.Directory, "Saved", String.Format("Manifest-{0}-{1}-{2}-{3}.xml", TargetName, Platform, Configuration, Arch));
 			ManifestFileNames.Add(ManifestFileName);
-			string Arguments = String.Format("-plugin={0} -iwyu -noubtmakefiles -manifest={1} -nohotreload", CommandUtils.MakePathSafeToUseWithCommandLine(HostProjectPluginFile.FullName), CommandUtils.MakePathSafeToUseWithCommandLine(ManifestFileName.FullName));
+			string Arguments = String.Format("-plugin={0} -noubtmakefiles -manifest={1} -nohotreload", CommandUtils.MakePathSafeToUseWithCommandLine(HostProjectPluginFile.FullName), CommandUtils.MakePathSafeToUseWithCommandLine(ManifestFileName.FullName));
 			Arguments += String.Format(" -Architecture={0}", Arch);
 			if (!String.IsNullOrEmpty(InAdditionalArgs))
 			{
@@ -43,7 +43,7 @@ public class BuildPlugin_HoloLens : BuildPlugin.TargetPlatform
 		{
 			FileReference ManifestFileName = FileReference.Combine(HostProjectFile.Directory, "Saved", String.Format("Manifest-{0}-{1}-{2}-{3}.xml", TargetName, Platform, Configuration, Arch));
 			ManifestFileNames.Add(ManifestFileName);
-			string Arguments = String.Format("-plugin={0} -iwyu -noubtmakefiles -manifest={1} -nohotreload", CommandUtils.MakePathSafeToUseWithCommandLine(HostProjectPluginFile.FullName), CommandUtils.MakePathSafeToUseWithCommandLine(ManifestFileName.FullName));
+			string Arguments = String.Format("-plugin={0} -noubtmakefiles -manifest={1} -nohotreload", CommandUtils.MakePathSafeToUseWithCommandLine(HostProjectPluginFile.FullName), CommandUtils.MakePathSafeToUseWithCommandLine(ManifestFileName.FullName));
 			Arguments += String.Format(" -Architecture={0}", Arch);
 			if (!String.IsNullOrEmpty(InAdditionalArgs))
 			{
diff --git a/Engine/Source/Programs/AutomationTool/Scripts/BuildPluginCommand.Automation.cs b/Engine/Source/Programs/AutomationTool/Scripts/BuildPluginCommand.Automation.cs
--- a/Engine/Source/Programs/AutomationTool/Scripts/BuildPluginCommand.Automation.cs
+++ b/Engine/Source/Programs/AutomationTool/Scripts/BuildPluginCommand.Automation.cs
@@ -305,7 +305,7 @@ public sealed class BuildPlugin : BuildCommand
 				FileReference ManifestFileName = FileReference.Combine(HostProjectFile.Directory, "Saved", String.Format("Manifest-{0}-{1}-{2}.xml", TargetName, Platform, Configuration));
 				ManifestFileNames.Add(ManifestFileName);
 				
-				string Arguments = String.Format("-plugin={0} -iwyu -noubtmakefiles -manifest={1} -nohotreload", CommandUtils.MakePathSafeToUseWithCommandLine(HostProjectPluginFile.FullName), CommandUtils.MakePathSafeToUseWithCommandLine(ManifestFileName.FullName));
+				string Arguments = String.Format("-plugin={0} -noubtmakefiles -manifest={1} -nohotreload", CommandUtils.MakePathSafeToUseWithCommandLine(HostProjectPluginFile.FullName), CommandUtils.MakePathSafeToUseWithCommandLine(ManifestFileName.FullName));
 
 				if (!String.IsNullOrEmpty(InAdditionalArgs))
 				{
@@ -349,7 +349,7 @@ public sealed class BuildPlugin : BuildCommand
 				FileReference ManifestFileName = FileReference.Combine(HostProjectFile.Directory, "Saved", String.Format("Manifest-{0}-{1}-{2}.xml", TargetName, Platform, Configuration));
 				ManifestFileNames.Add(ManifestFileName);
 				
-				string Arguments = String.Format("-plugin={0} -iwyu -noubtmakefiles -manifest={1} -nohotreload", CommandUtils.MakePathSafeToUseWithCommandLine(HostProjectPluginFile.FullName), CommandUtils.MakePathSafeToUseWithCommandLine(ManifestFileName.FullName));
+				string Arguments = String.Format("-plugin={0} -noubtmakefiles -manifest={1} -nohotreload", CommandUtils.MakePathSafeToUseWithCommandLine(HostProjectPluginFile.FullName), CommandUtils.MakePathSafeToUseWithCommandLine(ManifestFileName.FullName));
 
 				if (PlatformToArchitectureMap.TryGetValue(Platform, out string SpecifiedArchitecture) && !string.IsNullOrEmpty(SpecifiedArchitecture))
 				{
