- Another HoloLens compiler fix for allowing XmlConfig and CommandLine overrides for compiler version (to specify 2019 compiler). Only project settings were being respected before
- Apply commandline options to UnrealArchitectureConfig objects in Build and GenProject modes
- Also fixed a commant for XML documentation

diff --git a/Engine/Platforms/Hololens/Source/Programs/UnrealBuildTool/UEBuildHoloLens.cs b/Engine/Platforms/Hololens/Source/Programs/UnrealBuildTool/UEBuildHoloLens.cs
--- a/Engine/Platforms/Hololens/Source/Programs/UnrealBuildTool/UEBuildHoloLens.cs
+++ b/Engine/Platforms/Hololens/Source/Programs/UnrealBuildTool/UEBuildHoloLens.cs
@@ -205,19 +205,22 @@ namespace UnrealBuildTool
 
 	class HoloLensArchitectureConfig : UnrealArchitectureConfig
 	{
-		// we need to get the compiler the Target will use, before the target is created, when making projectfiles especially,
-		// so we duplicate some logic elsewhere in this file
-		[ConfigFile(ConfigHierarchyType.Engine, "/Script/HoloLensPlatformEditor.HoloLensTargetSettings", "CompilerVersion")]
-		[XmlConfigFile(Category = "HoloLensPlatform")]
+		/// <summary>
+		/// Duplicating the commandline param so that we can apply the cmdline to this object in BuildMode and GenProjFilesMode
+		/// </summary>
 		[CommandLine("-2019", Value = nameof(WindowsCompiler.VisualStudio2019))]
 		[CommandLine("-2022", Value = nameof(WindowsCompiler.VisualStudio2022))]
 		[CommandLine("-Compiler=")]
-		private WindowsCompiler Compiler = WindowsCompiler.Default;
+		private WindowsCompiler CommandLineCompiler = WindowsCompiler.Default;
 
 		private WindowsCompiler GetCompiler(FileReference? ProjectFile, UnrealArch Architecture)
 		{
+			// we don't have access to the commandline Arguments here, so we can't run it on the HoloLensPlatform object now
+			// we count on the Mode applying the Arguments to this object (and all ArchConfig objects) so that CommandLineCompiler
+			// is set below if needed
 			HoloLensTargetRules HoloLensPlatform = new();
 			ConfigCache.ReadSettings(ProjectFile?.Directory, UnrealTargetPlatform.HoloLens, HoloLensPlatform);
+			XmlConfig.ApplyTo(HoloLensPlatform);
 
 			// HoloLensPlatform has logic that the value in ReadSettings will override whatever the commandline specified version is,
 			// which is hard to say if it's desired, but need to stay consistent (could move this into a function)
@@ -226,9 +229,9 @@ namespace UnrealBuildTool
 				return HoloLensPlatform.Compiler;
 			}
 
-			if (Compiler != WindowsCompiler.Default)
+			if (CommandLineCompiler != WindowsCompiler.Default)
 			{
-				return Compiler;
+				return CommandLineCompiler;
 			}
 
 			// get the default compiler, but don't spit out a nasty warning if it doesn't exist
@@ -251,7 +254,10 @@ namespace UnrealBuildTool
 
 			bool bHasArmToolchain = MicrosoftPlatformSDK.HasValidCompiler(GetCompiler(ProjectFile, UnrealArch.Arm64), UnrealArch.Arm64, Logger);
 			bool bHasX64Toolchain = MicrosoftPlatformSDK.HasValidCompiler(GetCompiler(ProjectFile, UnrealArch.X64), UnrealArch.X64, Logger);
-	
+
+			Log.TraceLog($"HoloLens Arm: Desired Compiler: {GetCompiler(ProjectFile, UnrealArch.Arm64)}, Valid: {bHasArmToolchain}");
+			Log.TraceLog($"HoloLens X64: Desired Compiler: {GetCompiler(ProjectFile, UnrealArch.X64)}, Valid: {bHasX64Toolchain}");
+
 			ConfigHierarchy Ini = ConfigCache.ReadHierarchy(ConfigHierarchyType.Engine, ProjectFile?.Directory, UnrealTargetPlatform.HoloLens);
 			if (Ini.GetBool("/Script/HoloLensPlatformEditor.HoloLensTargetSettings", "bBuildForEmulation", out bBuildForEmulation) && bBuildForEmulation)
 			{
diff --git a/Engine/Source/Programs/UnrealBuildTool/Configuration/UEBuildPlatform.cs b/Engine/Source/Programs/UnrealBuildTool/Configuration/UEBuildPlatform.cs
--- a/Engine/Source/Programs/UnrealBuildTool/Configuration/UEBuildPlatform.cs
+++ b/Engine/Source/Programs/UnrealBuildTool/Configuration/UEBuildPlatform.cs
@@ -55,6 +55,16 @@ namespace UnrealBuildTool
 			return UEBuildPlatform.GetBuildPlatform(Platform).ArchitectureConfig;
 		}
 
+		/// <summary>
+		/// Get all known ArchitectureConfig objects
+		/// </summary>
+		/// <returns></returns>
+		public static IEnumerable<UnrealArchitectureConfig> AllConfigs()
+		{
+			// return ArchConfigs for platforms that have a BuildPlatform (which is where the Configs are stored)
+			return UnrealTargetPlatform.GetValidPlatforms().Where(x => UEBuildPlatform.TryGetBuildPlatform(x, out _)).Select(x => UEBuildPlatform.GetBuildPlatform(x).ArchitectureConfig);
+		}
+
 
 		/// <summary>
 		/// The multi-architecture mode for this platform (potentially single-architecture)
diff --git a/Engine/Source/Programs/UnrealBuildTool/Modes/BuildMode.cs b/Engine/Source/Programs/UnrealBuildTool/Modes/BuildMode.cs
--- a/Engine/Source/Programs/UnrealBuildTool/Modes/BuildMode.cs
+++ b/Engine/Source/Programs/UnrealBuildTool/Modes/BuildMode.cs
@@ -119,6 +119,12 @@ namespace UnrealBuildTool
 			// Read the XML configuration files
 			XmlConfig.ApplyTo(this);
 
+			// Apply to architecture configs that need to read commandline arguments and didn't have the Arguments passed in during construction
+			foreach (UnrealArchitectureConfig Config in UnrealArchitectureConfig.AllConfigs())
+			{
+				Arguments.ApplyTo(Config);
+			}
+
 			// Fixup the log path if it wasn't overridden by a config file
 			if (BaseLogFileName == null)
 			{
diff --git a/Engine/Source/Programs/UnrealBuildTool/Modes/GenerateProjectFilesMode.cs b/Engine/Source/Programs/UnrealBuildTool/Modes/GenerateProjectFilesMode.cs
--- a/Engine/Source/Programs/UnrealBuildTool/Modes/GenerateProjectFilesMode.cs
+++ b/Engine/Source/Programs/UnrealBuildTool/Modes/GenerateProjectFilesMode.cs
@@ -69,6 +69,12 @@ namespace UnrealBuildTool
 			// Apply the XML config to this class
 			XmlConfig.ApplyTo(this);
 
+			// Apply to architecture configs that need to read commandline arguments and didn't have the Arguments passed in during construction
+			foreach (UnrealArchitectureConfig Config in UnrealArchitectureConfig.AllConfigs())
+			{
+				Arguments.ApplyTo(Config);
+			}
+
 			// set up logging (taken from BuildMode)
 			FileReference LogFile = FileReference.Combine(Unreal.EngineProgramSavedDirectory, "UnrealBuildTool", "Log_GPF.txt");
 			Log.AddFileWriter("DefaultLogTraceListener", LogFile);
