- Added MarkUsed flag to CommandLineAttribute
- Fixed up how UnrealArchitectureConfig classes have commandline applied
- Proper fix for UE-176691

diff --git a/Engine/Platforms/Hololens/Source/Programs/UnrealBuildTool/UEBuildHoloLens.cs b/Engine/Platforms/Hololens/Source/Programs/UnrealBuildTool/UEBuildHoloLens.cs
--- a/Engine/Platforms/Hololens/Source/Programs/UnrealBuildTool/UEBuildHoloLens.cs
+++ b/Engine/Platforms/Hololens/Source/Programs/UnrealBuildTool/UEBuildHoloLens.cs
@@ -208,9 +208,9 @@ namespace UnrealBuildTool
 		/// <summary>
 		/// Duplicating the commandline param so that we can apply the cmdline to this object in BuildMode and GenProjFilesMode
 		/// </summary>
-		[CommandLine("-2019", Value = nameof(WindowsCompiler.VisualStudio2019))]
-		[CommandLine("-2022", Value = nameof(WindowsCompiler.VisualStudio2022))]
-		[CommandLine("-Compiler=")]
+		[CommandLine("-2019", Value = nameof(WindowsCompiler.VisualStudio2019), MarkUsed = false)]
+		[CommandLine("-2022", Value = nameof(WindowsCompiler.VisualStudio2022), MarkUsed = false)]
+		[CommandLine("-Compiler=", MarkUsed = false)]
 		private WindowsCompiler CommandLineCompiler = WindowsCompiler.Default;
 
 		private WindowsCompiler GetCompiler(FileReference? ProjectFile, UnrealArch Architecture)
@@ -255,8 +255,8 @@ namespace UnrealBuildTool
 			bool bHasArmToolchain = MicrosoftPlatformSDK.HasValidCompiler(GetCompiler(ProjectFile, UnrealArch.Arm64), UnrealArch.Arm64, Logger);
 			bool bHasX64Toolchain = MicrosoftPlatformSDK.HasValidCompiler(GetCompiler(ProjectFile, UnrealArch.X64), UnrealArch.X64, Logger);
 
-			Log.TraceLog($"HoloLens Arm: Desired Compiler: {GetCompiler(ProjectFile, UnrealArch.Arm64)}, Valid: {bHasArmToolchain}");
-			Log.TraceLog($"HoloLens X64: Desired Compiler: {GetCompiler(ProjectFile, UnrealArch.X64)}, Valid: {bHasX64Toolchain}");
+			Log.TraceLogOnce($"HoloLens Arm: Desired Compiler: {GetCompiler(ProjectFile, UnrealArch.Arm64)}, Valid: {bHasArmToolchain}");
+			Log.TraceLogOnce($"HoloLens X64: Desired Compiler: {GetCompiler(ProjectFile, UnrealArch.X64)}, Valid: {bHasX64Toolchain}");
 
 			ConfigHierarchy Ini = ConfigCache.ReadHierarchy(ConfigHierarchyType.Engine, ProjectFile?.Directory, UnrealTargetPlatform.HoloLens);
 			if (Ini.GetBool("/Script/HoloLensPlatformEditor.HoloLensTargetSettings", "bBuildForEmulation", out bBuildForEmulation) && bBuildForEmulation)
diff --git a/Engine/Source/Programs/Shared/EpicGames.Core/CommandLineArguments.cs b/Engine/Source/Programs/Shared/EpicGames.Core/CommandLineArguments.cs
--- a/Engine/Source/Programs/Shared/EpicGames.Core/CommandLineArguments.cs
+++ b/Engine/Source/Programs/Shared/EpicGames.Core/CommandLineArguments.cs
@@ -789,7 +789,10 @@ namespace EpicGames.Core
 							}
 
 							// Mark this argument as used
-							_usedArguments.Set(index, true);
+							if (attribute.MarkUsed)
+							{
+								_usedArguments.Set(index, true);
+							}
 						}
 					}
 
diff --git a/Engine/Source/Programs/Shared/EpicGames.Core/CommandLineAttribute.cs b/Engine/Source/Programs/Shared/EpicGames.Core/CommandLineAttribute.cs
--- a/Engine/Source/Programs/Shared/EpicGames.Core/CommandLineAttribute.cs
+++ b/Engine/Source/Programs/Shared/EpicGames.Core/CommandLineAttribute.cs
@@ -27,6 +27,11 @@ namespace EpicGames.Core
 		/// </summary>
 		public bool Required { get; set; }
 
+		/// <summary>
+		/// Whether this argument is required
+		/// </summary>
+		public bool MarkUsed { get; set; } = true;
+
 		/// <summary>
 		/// For collection types, specifies the separator character between multiple arguments
 		/// </summary>
diff --git a/Engine/Source/Programs/UnrealBuildTool/Configuration/TargetRules.cs b/Engine/Source/Programs/UnrealBuildTool/Configuration/TargetRules.cs
--- a/Engine/Source/Programs/UnrealBuildTool/Configuration/TargetRules.cs
+++ b/Engine/Source/Programs/UnrealBuildTool/Configuration/TargetRules.cs
@@ -2559,9 +2559,13 @@ namespace UnrealBuildTool
 					}
 				}
 			}
-			foreach (UnrealArch Arch in Architectures.Architectures)
+			foreach (UnrealTargetPlatform Platform in UnrealTargetPlatform.GetValidPlatforms())
 			{
-				yield return Arch;
+				UEBuildPlatform? BuildPlatform;
+				if (UEBuildPlatform.TryGetBuildPlatform(Platform, out BuildPlatform))
+				{
+					yield return BuildPlatform.ArchitectureConfig;
+				}
 			}
 		}
 
diff --git a/Engine/Source/Programs/UnrealBuildTool/Configuration/UEBuildPlatform.cs b/Engine/Source/Programs/UnrealBuildTool/Configuration/UEBuildPlatform.cs
--- a/Engine/Source/Programs/UnrealBuildTool/Configuration/UEBuildPlatform.cs
+++ b/Engine/Source/Programs/UnrealBuildTool/Configuration/UEBuildPlatform.cs
@@ -55,6 +55,17 @@ namespace UnrealBuildTool
 			return UEBuildPlatform.GetBuildPlatform(Platform).ArchitectureConfig;
 		}
 
+		/// <summary>
+		/// Get all known ArchitectureConfig objects
+		/// </summary>
+		/// <returns></returns>
+		public static IEnumerable<UnrealArchitectureConfig> AllConfigs()
+		{
+			// return ArchConfigs for platforms that have a BuildPlatform (which is where the Configs are stored)
+			return UnrealTargetPlatform.GetValidPlatforms().Where(x => UEBuildPlatform.TryGetBuildPlatform(x, out _)).Select(x => UEBuildPlatform.GetBuildPlatform(x).ArchitectureConfig);
+		}
+
+
 		/// <summary>
 		/// The multi-architecture mode for this platform (potentially single-architecture)
 		/// </summary>
diff --git a/Engine/Source/Programs/UnrealBuildTool/Modes/BuildMode.cs b/Engine/Source/Programs/UnrealBuildTool/Modes/BuildMode.cs
--- a/Engine/Source/Programs/UnrealBuildTool/Modes/BuildMode.cs
+++ b/Engine/Source/Programs/UnrealBuildTool/Modes/BuildMode.cs
@@ -119,6 +119,12 @@ namespace UnrealBuildTool
 			// Read the XML configuration files
 			XmlConfig.ApplyTo(this);
 
+			// Apply to architecture configs that need to read commandline arguments and didn't have the Arguments passed in during construction
+			foreach (UnrealArchitectureConfig Config in UnrealArchitectureConfig.AllConfigs())
+			{
+				Arguments.ApplyTo(Config);
+			}
+
 			// Fixup the log path if it wasn't overridden by a config file
 			if (BaseLogFileName == null)
 			{
diff --git a/Engine/Source/Programs/UnrealBuildTool/Modes/GenerateProjectFilesMode.cs b/Engine/Source/Programs/UnrealBuildTool/Modes/GenerateProjectFilesMode.cs
--- a/Engine/Source/Programs/UnrealBuildTool/Modes/GenerateProjectFilesMode.cs
+++ b/Engine/Source/Programs/UnrealBuildTool/Modes/GenerateProjectFilesMode.cs
@@ -69,6 +69,12 @@ namespace UnrealBuildTool
 			// Apply the XML config to this class
 			XmlConfig.ApplyTo(this);
 
+			// Apply to architecture configs that need to read commandline arguments and didn't have the Arguments passed in during construction
+			foreach (UnrealArchitectureConfig Config in UnrealArchitectureConfig.AllConfigs())
+			{
+				Arguments.ApplyTo(Config);
+			}
+
 			// set up logging (taken from BuildMode)
 			FileReference LogFile = FileReference.Combine(Unreal.EngineProgramSavedDirectory, "UnrealBuildTool", "Log_GPF.txt");
 			Log.AddFileWriter("DefaultLogTraceListener", LogFile);
