Cleaning up some D3D11RHI dependencies.

diff --git a/Engine/Platforms/Hololens/Source/Runtime/Windows/D3D11RHI/D3D11RHI_HoloLens.build.cs b/Engine/Platforms/Hololens/Source/Runtime/Windows/D3D11RHI/D3D11RHI_HoloLens.build.cs
--- a/Engine/Platforms/Hololens/Source/Runtime/Windows/D3D11RHI/D3D11RHI_HoloLens.build.cs
+++ b/Engine/Platforms/Hololens/Source/Runtime/Windows/D3D11RHI/D3D11RHI_HoloLens.build.cs
@@ -8,7 +8,6 @@ namespace UnrealBuildTool.Rules
 
 		public D3D11RHI_HoloLens(ReadOnlyTargetRules Target) : base(Target)
 		{
-			PrivateIncludePaths.Add("../Platforms/HoloLens/Source/Runtime/Windows/D3D11RHI/Private");
 		}
 	}
 }
diff --git a/Engine/Platforms/Hololens/Source/Runtime/Windows/D3D11RHI/Private/D3D11StateCache.h b/Engine/Platforms/Hololens/Source/Runtime/Windows/D3D11RHI/Private/D3D11StateCache.h
deleted file mode 100644
diff --git a/Engine/Platforms/Hololens/Source/Runtime/Windows/D3D11RHI/Private/HoloLensD3D11ConstantBuffer.cpp b/Engine/Platforms/Hololens/Source/Runtime/Windows/D3D11RHI/Private/HoloLensD3D11ConstantBuffer.cpp
--- a/Engine/Platforms/Hololens/Source/Runtime/Windows/D3D11RHI/Private/HoloLensD3D11ConstantBuffer.cpp
+++ b/Engine/Platforms/Hololens/Source/Runtime/Windows/D3D11RHI/Private/HoloLensD3D11ConstantBuffer.cpp
@@ -1,126 +1,6 @@
 // Copyright Epic Games, Inc. All Rights Reserved.
 
-/*=============================================================================
-	HoloLensD3D11ConstantBuffer.cpp: D3D Constant Buffer functions
-=============================================================================*/
-
 #include "D3D11RHIPrivate.h"
-#include "HoloLensD3D11ConstantBuffer.h"
-
-void FWinD3D11ConstantBuffer::InitDynamicRHI()
-{
-	TRefCountPtr<ID3D11Buffer> CBuffer = NULL;
-	D3D11_BUFFER_DESC BufferDesc;
-	ZeroMemory( &BufferDesc, sizeof( D3D11_BUFFER_DESC ) );
-	// Verify constant buffer ByteWidth requirements
-	check(MaxSize <= D3D11_REQ_CONSTANT_BUFFER_ELEMENT_COUNT && (MaxSize % 16) == 0);
-	BufferDesc.ByteWidth = MaxSize;
-
-	BufferDesc.Usage = D3D11_USAGE_DEFAULT;
-	BufferDesc.CPUAccessFlags = 0;
-	BufferDesc.BindFlags = D3D11_BIND_CONSTANT_BUFFER;
-	BufferDesc.MiscFlags = 0;
-
-	Buffers = new TRefCountPtr<ID3D11Buffer>[NumSubBuffers];
-	CurrentSubBuffer = 0;
-	for(uint32 s = 0;s < NumSubBuffers;s++)
-	{
-		VERIFYD3D11RESULT(D3DRHI->GetDevice()->CreateBuffer(&BufferDesc, NULL, Buffers[s].GetInitReference()));
-		UpdateBufferStats(Buffers[s],true);
-		BufferDesc.ByteWidth = Align(BufferDesc.ByteWidth / 2, 16);
-	}
-
-	FD3D11ConstantBuffer::InitDynamicRHI();
-}
-
-void FWinD3D11ConstantBuffer::ReleaseDynamicRHI()
-{
-	FD3D11ConstantBuffer::ReleaseDynamicRHI();
-
-	if(Buffers)
-	{
-		for(uint32 s = 0;s < NumSubBuffers;s++)
-		{
-			UpdateBufferStats(Buffers[s],false);
-		}
-
-		delete [] Buffers;
-		Buffers = NULL;
-	}
-}
-
-bool FWinD3D11ConstantBuffer::CommitConstantsToDevice( bool bDiscardSharedConstants )
-{
-	// New circular buffer system for faster constant uploads.  Avoids CopyResource and speeds things up considerably
-	if (CurrentUpdateSize == 0)
-	{
-		return false;
-	}
-
-	SCOPE_CYCLE_COUNTER(STAT_D3D11GlobalConstantBufferUpdateTime);
-
-	if ( bDiscardSharedConstants )
-	{
-		// If we're discarding shared constants, just use constants that have been updated since the last Commit.
-		TotalUpdateSize = CurrentUpdateSize;
-	}
-	else
-	{
-		// If we're re-using shared constants, use all constants.
-		TotalUpdateSize = FMath::Max( CurrentUpdateSize, TotalUpdateSize );
-	}
-
-	// This basically keeps track dynamically how much data has been updated every frame
-	// and then divides up a "max" constant buffer size by halves down until it finds a large sections that more tightly covers
-	// the amount updated, assuming that all data in a constant buffer is updated each draw call and contiguous.
-	// This only used for the IndexSlot==0 constant buffer on the vertex shader.
-	// Which will have an indeterminate number of constant values that are generated by material shaders.
-	CurrentSubBuffer = 1;
-	uint32 BufferSize = MaxSize / 2;
-	while ( BufferSize >= TotalUpdateSize && CurrentSubBuffer < NumSubBuffers )
-	{
-		CurrentSubBuffer++;
-		BufferSize /= 2;
-	}
-	CurrentSubBuffer--;
-	BufferSize *= 2;
-
-	ID3D11Buffer* Buffer = Buffers[CurrentSubBuffer];
-	check(((uint64)(ShadowData) & 0xf) == 0);
-	D3DRHI->GetDeviceContext()->UpdateSubresource(Buffer,0,NULL,(void*)ShadowData,BufferSize,BufferSize);
-
-	CurrentUpdateSize = 0;
-
-	return true;
-}
-
-void FD3D11DynamicRHI::InitConstantBuffers()
-{
-	// Allocate shader constant buffers.  All shader types can have access to all buffers.
-	//  Index==0 is reserved for "custom" params, and the rest are reserved by the system for Common
-	//	constants
-	VSConstantBuffers.Empty(MAX_CONSTANT_BUFFER_SLOTS);
-	PSConstantBuffers.Empty(MAX_CONSTANT_BUFFER_SLOTS);
-	GSConstantBuffers.Empty(MAX_CONSTANT_BUFFER_SLOTS);
-	CSConstantBuffers.Empty(MAX_CONSTANT_BUFFER_SLOTS);
-	for(int32 BufferIndex = 0;BufferIndex < MAX_CONSTANT_BUFFER_SLOTS;BufferIndex++)
-	{
-		uint32 Size = GConstantBufferSizes[BufferIndex];
-		uint32 SubBuffers = 1;
-		if(BufferIndex == GLOBAL_CONSTANT_BUFFER_INDEX)
-		{
-			SubBuffers = 5;
-		}
-
-		// Vertex shader can have subbuffers for index==0.  This is from Epic's original design for the auto-fit of size to
-		//	reduce the update costs of the buffer.
-
-		// New circular buffer system for faster constant uploads.  Avoids CopyResource and speeds things up considerably
-		VSConstantBuffers.Add(new FWinD3D11ConstantBuffer(this, Size, SubBuffers));
-		PSConstantBuffers.Add(new FWinD3D11ConstantBuffer(this, Size, SubBuffers));
-		GSConstantBuffers.Add(new FWinD3D11ConstantBuffer(this, Size));
-		CSConstantBuffers.Add(new FWinD3D11ConstantBuffer(this, Size));
-	}
-}
 
-DEFINE_STAT(STAT_D3D11GlobalConstantBufferUpdateTime);
+// Borrow the Windows desktop version
+#include "Windows/WindowsD3D11ConstantBuffer.cpp"
diff --git a/Engine/Platforms/Hololens/Source/Runtime/Windows/D3D11RHI/Private/HoloLensD3D11ConstantBuffer.h b/Engine/Platforms/Hololens/Source/Runtime/Windows/D3D11RHI/Private/HoloLensD3D11ConstantBuffer.h
--- a/Engine/Platforms/Hololens/Source/Runtime/Windows/D3D11RHI/Private/HoloLensD3D11ConstantBuffer.h
+++ b/Engine/Platforms/Hololens/Source/Runtime/Windows/D3D11RHI/Private/HoloLensD3D11ConstantBuffer.h
@@ -1,41 +1,6 @@
 // Copyright Epic Games, Inc. All Rights Reserved.
 
-/*=============================================================================
-	HoloLensD3D11ConstantBuffer.h: D3D Constant Buffer functions
-=============================================================================*/
-
 #pragma once
 
-class FWinD3D11ConstantBuffer : public FD3D11ConstantBuffer
-{
-public:
-	FWinD3D11ConstantBuffer(FD3D11DynamicRHI* InD3DRHI, uint32 InSize = 0, uint32 SubBuffers = 1) :
-		FD3D11ConstantBuffer(InD3DRHI, InSize, SubBuffers),
-		Buffers(nullptr),
-		CurrentSubBuffer(0),
-		NumSubBuffers(SubBuffers)
-	{
-	}
-
-	// FRenderResource interface.
-	virtual void	InitDynamicRHI() override;
-	virtual void	ReleaseDynamicRHI() override;
-
-	/**
-	* Get the current pool buffer
-	*/
-	ID3D11Buffer* GetConstantBuffer() const
-	{
-		return Buffers[CurrentSubBuffer];
-	}
-
-	/**
-	* Unlocks the constant buffer so the data can be transmitted to the device
-	*/
-	bool CommitConstantsToDevice(bool bDiscardSharedConstants);
-
-private:
-	TRefCountPtr<ID3D11Buffer>* Buffers;
-	uint32	CurrentSubBuffer;
-	uint32	NumSubBuffers;
-};
+// Borrow the Windows desktop version
+#include "Windows/WindowsD3D11ConstantBuffer.h"
diff --git a/Engine/Source/Runtime/Windows/D3D11RHI/Private/D3D11StateCache.cpp b/Engine/Source/Runtime/Windows/D3D11RHI/Private/D3D11StateCache.cpp
--- a/Engine/Source/Runtime/Windows/D3D11RHI/Private/D3D11StateCache.cpp
+++ b/Engine/Source/Runtime/Windows/D3D11RHI/Private/D3D11StateCache.cpp
@@ -7,7 +7,7 @@
 //	Include Files
 //-----------------------------------------------------------------------------
 #include "D3D11RHIPrivate.h"
-#include "Windows/D3D11StateCache.h"
+#include "D3D11StateCache.h"
 
 #if D3D11_ALLOW_STATE_CACHE && D3D11_STATE_CACHE_RUNTIME_TOGGLE
 
diff --git a/Engine/Source/Runtime/Windows/D3D11RHI/Private/D3D11StateCache.h b/Engine/Source/Runtime/Windows/D3D11RHI/Private/D3D11StateCache.h
new file mode 100644
--- /dev/null
+++ b/Engine/Source/Runtime/Windows/D3D11RHI/Private/D3D11StateCache.h
@@ -0,0 +1,10 @@
+// Copyright Epic Games, Inc. All Rights Reserved.
+
+// Implementation of Device Context State Caching to improve draw
+//	thread performance by removing redundant device context calls.
+
+#pragma once
+
+#include "D3D11StateCachePrivate.h"
+
+typedef FD3D11StateCacheBase FD3D11StateCache;
diff --git a/Engine/Source/Runtime/Windows/D3D11RHI/Private/Windows/D3D11StateCache.h b/Engine/Source/Runtime/Windows/D3D11RHI/Private/Windows/D3D11StateCache.h
deleted file mode 100644
