Fixing some D3D11 include paths.

diff --git a/Engine/Platforms/Hololens/Source/Runtime/Windows/D3D11RHI/Private/D3D11RHIPrivateUtil.h b/Engine/Platforms/Hololens/Source/Runtime/Windows/D3D11RHI/Private/D3D11RHIPrivateUtil.h
--- a/Engine/Platforms/Hololens/Source/Runtime/Windows/D3D11RHI/Private/D3D11RHIPrivateUtil.h
+++ b/Engine/Platforms/Hololens/Source/Runtime/Windows/D3D11RHI/Private/D3D11RHIPrivateUtil.h
@@ -1,24 +1,5 @@
 // Copyright Epic Games, Inc. All Rights Reserved.
 
-/*=============================================================================
-	D3D11RHIPrivateUtil.h: Private D3D RHI Utility definitions for HoloLens.
-=============================================================================*/
-
 #pragma once
 
-#include "HoloLensD3D11ConstantBuffer.h"
-
-struct FD3DRHIUtil
-{
-	template <EShaderFrequency ShaderFrequencyT>
-	static FORCEINLINE void CommitConstants(FD3D11ConstantBuffer* InConstantBuffer, FD3D11StateCache& StateCache, uint32 Index, bool bDiscardSharedConstants)
-	{
-		auto* ConstantBuffer = ((FWinD3D11ConstantBuffer*)InConstantBuffer);
-		// Array may contain NULL entries to pad out to proper 
-		if (ConstantBuffer && ConstantBuffer->CommitConstantsToDevice(bDiscardSharedConstants))
-		{
-			ID3D11Buffer* DeviceBuffer = ConstantBuffer->GetConstantBuffer();
-			StateCache.SetConstantBuffer<ShaderFrequencyT>(DeviceBuffer, Index);
-		}
-	}
-};
+#include "Windows/D3D11RHIPrivateUtil.h"
diff --git a/Engine/Source/Runtime/Windows/D3D11RHI/Private/D3D11Commands.cpp b/Engine/Source/Runtime/Windows/D3D11RHI/Private/D3D11Commands.cpp
--- a/Engine/Source/Runtime/Windows/D3D11RHI/Private/D3D11Commands.cpp
+++ b/Engine/Source/Runtime/Windows/D3D11RHI/Private/D3D11Commands.cpp
@@ -5,7 +5,7 @@
 =============================================================================*/
 
 #include "D3D11RHIPrivate.h"
-#include "D3D11RHIPrivateUtil.h"
+#include "Windows/D3D11RHIPrivateUtil.h"
 #include "StaticBoundShaderState.h"
 #include "GlobalShader.h"
 #include "OneColorShader.h"
