Removing hardcoded Platforms/ include paths.

diff --git a/Engine/Platforms/Hololens/Source/Developer/HoloLensTargetPlatform/HoloLensTargetPlatform.Build.cs b/Engine/Platforms/Hololens/Source/Developer/HoloLensTargetPlatform/HoloLensTargetPlatform.Build.cs
--- a/Engine/Platforms/Hololens/Source/Developer/HoloLensTargetPlatform/HoloLensTargetPlatform.Build.cs
+++ b/Engine/Platforms/Hololens/Source/Developer/HoloLensTargetPlatform/HoloLensTargetPlatform.Build.cs
@@ -1,7 +1,7 @@
 // Copyright Epic Games, Inc. All Rights Reserved.
 
+using EpicGames.Core;
 using System;
-using System.IO;
 using UnrealBuildTool;
 
 public class HoloLensTargetPlatform : ModuleRules
@@ -21,10 +21,11 @@ public class HoloLensTargetPlatform : ModuleRules
             }
         );
 
-		// directly include the location of the HoloLensPlatformProperties.h files so that we don't instantiate the Core_HoloLens module for host platform
-		PrivateIncludePaths.Add(Path.Combine(EngineDirectory, "Platforms/HoloLens/Source/Runtime/Core/Public"));
-
-		PrivateIncludePathModuleNames.Add("Settings");
+		PrivateIncludePathModuleNames.AddAll(
+			"Settings",
+			// Needed for HoloLensPlatformProperties.h
+			"Core_HoloLens"
+		);
 
         int Win10Build = 0;
         Version ver = null;
diff --git a/Engine/Platforms/Hololens/Source/Runtime/Core/Core_HoloLens.build.cs b/Engine/Platforms/Hololens/Source/Runtime/Core/Core_HoloLens.build.cs
--- a/Engine/Platforms/Hololens/Source/Runtime/Core/Core_HoloLens.build.cs
+++ b/Engine/Platforms/Hololens/Source/Runtime/Core/Core_HoloLens.build.cs
@@ -6,18 +6,20 @@ namespace UnrealBuildTool.Rules
 	{
 		public Core_HoloLens(ReadOnlyTargetRules Target) : base(Target)
 		{
-			PublicIncludePaths.Add("../Platforms/HoloLens/Source/Runtime/Core/Public");
-			AddEngineThirdPartyPrivateStaticDependencies(Target,
-				"zlib");
+			// This module needs to be referenced from a TargetPlatform module, so don't leak our platform info out
+			if (Target.Platform == UnrealTargetPlatform.HoloLens)
+			{
+				AddEngineThirdPartyPrivateStaticDependencies(Target,
+					"zlib");
 
-			AddEngineThirdPartyPrivateStaticDependencies(Target,
-				"IntelTBB",
-				"XInput"
-				);
-
-			PublicDefinitions.Add("WITH_VS_PERF_PROFILER=0");
-			PublicDefinitions.Add("IS_RUNNING_GAMETHREAD_ON_EXTERNAL_THREAD=1");
+				AddEngineThirdPartyPrivateStaticDependencies(Target,
+					"IntelTBB",
+					"XInput"
+					);
 
+				PublicDefinitions.Add("WITH_VS_PERF_PROFILER=0");
+				PublicDefinitions.Add("IS_RUNNING_GAMETHREAD_ON_EXTERNAL_THREAD=1");
+			}
 		}
 	}
 }
