UnrealBuildTool: Also check preview toolchain for hololens when generating a solution, for users with only the vs preview installed

#rnx

diff --git a/Engine/Platforms/Hololens/Source/Programs/UnrealBuildTool/HoloLensProjectGenerator.cs b/Engine/Platforms/Hololens/Source/Programs/UnrealBuildTool/HoloLensProjectGenerator.cs
--- a/Engine/Platforms/Hololens/Source/Programs/UnrealBuildTool/HoloLensProjectGenerator.cs
+++ b/Engine/Platforms/Hololens/Source/Programs/UnrealBuildTool/HoloLensProjectGenerator.cs
@@ -87,7 +87,12 @@ namespace UnrealBuildTool
 			DirectoryReference? PlatformWinMDLocation = HoloLensPlatform.GetCppCXMetadataLocation(Compiler, "Latest", UnrealArch.X64, Logger);
 			if (PlatformWinMDLocation == null || !FileReference.Exists(FileReference.Combine(PlatformWinMDLocation, "platform.winmd")))
 			{
-				Logger.LogWarning("Unable to find platform.winmd for {ToolChain} toolchain; '{Version}' is an invalid version", WindowsPlatform.GetCompilerName(Compiler), "Latest");
+				// Try Preview if Latest is not found
+				PlatformWinMDLocation = HoloLensPlatform.GetCppCXMetadataLocation(Compiler, "Preview", WindowsArchitecture.x64, Logger);
+				if (PlatformWinMDLocation == null || !FileReference.Exists(FileReference.Combine(PlatformWinMDLocation, "platform.winmd")))
+				{
+					Logger.LogWarning("Unable to find platform.winmd for {ToolChain} toolchain", WindowsPlatform.GetCompilerName(Compiler));
+				}
 			}
 			string FoundationWinMDPath = HoloLensPlatform.GetLatestMetadataPathForApiContract("Windows.Foundation.FoundationContract", Compiler, Logger);
 			string UniversalWinMDPath = HoloLensPlatform.GetLatestMetadataPathForApiContract("Windows.Foundation.UniversalApiContract", Compiler, Logger);
