UnrealBuildTool: Fix updated WindowsArchitecture -> UnrealArch reference

#rnx

diff --git a/Engine/Platforms/Hololens/Source/Programs/UnrealBuildTool/HoloLensProjectGenerator.cs b/Engine/Platforms/Hololens/Source/Programs/UnrealBuildTool/HoloLensProjectGenerator.cs
--- a/Engine/Platforms/Hololens/Source/Programs/UnrealBuildTool/HoloLensProjectGenerator.cs
+++ b/Engine/Platforms/Hololens/Source/Programs/UnrealBuildTool/HoloLensProjectGenerator.cs
@@ -88,7 +88,7 @@ namespace UnrealBuildTool
 			if (PlatformWinMDLocation == null || !FileReference.Exists(FileReference.Combine(PlatformWinMDLocation, "platform.winmd")))
 			{
 				// Try Preview if Latest is not found
-				PlatformWinMDLocation = HoloLensPlatform.GetCppCXMetadataLocation(Compiler, "Preview", WindowsArchitecture.x64, Logger);
+				PlatformWinMDLocation = HoloLensPlatform.GetCppCXMetadataLocation(Compiler, "Preview", UnrealArch.x64, Logger);
 				if (PlatformWinMDLocation == null || !FileReference.Exists(FileReference.Combine(PlatformWinMDLocation, "platform.winmd")))
 				{
 					Logger.LogWarning("Unable to find platform.winmd for {ToolChain} toolchain", WindowsPlatform.GetCompilerName(Compiler));
