Fix exception when generating a solution with only VS2022 Preview installed while the hololens platform is available. Requesting the "Latest" toolchain will error if the standard release VS is not available while passing null instead will allow returning a preview version and won't error if no toolchain is available.

#rnx

diff --git a/Engine/Platforms/Hololens/Source/Programs/UnrealBuildTool/HoloLensProjectGenerator.cs b/Engine/Platforms/Hololens/Source/Programs/UnrealBuildTool/HoloLensProjectGenerator.cs
--- a/Engine/Platforms/Hololens/Source/Programs/UnrealBuildTool/HoloLensProjectGenerator.cs
+++ b/Engine/Platforms/Hololens/Source/Programs/UnrealBuildTool/HoloLensProjectGenerator.cs
@@ -84,15 +84,10 @@ namespace UnrealBuildTool
 			}
 
 			WindowsCompiler Compiler = WindowsPlatform.GetDefaultCompiler(TargetFilePath, UnrealArch.X64, Logger);
-			DirectoryReference? PlatformWinMDLocation = HoloLensPlatform.GetCppCXMetadataLocation(Compiler, "Latest", UnrealArch.X64, Logger);
+			DirectoryReference? PlatformWinMDLocation = HoloLensPlatform.GetCppCXMetadataLocation(Compiler, null, UnrealArch.X64, Logger);
 			if (PlatformWinMDLocation == null || !FileReference.Exists(FileReference.Combine(PlatformWinMDLocation, "platform.winmd")))
 			{
-				// Try Preview if Latest is not found
-				PlatformWinMDLocation = HoloLensPlatform.GetCppCXMetadataLocation(Compiler, "Preview", UnrealArch.X64, Logger);
-				if (PlatformWinMDLocation == null || !FileReference.Exists(FileReference.Combine(PlatformWinMDLocation, "platform.winmd")))
-				{
-					Logger.LogWarning("Unable to find platform.winmd for {ToolChain} toolchain", WindowsPlatform.GetCompilerName(Compiler));
-				}
+				Logger.LogWarning("Unable to find platform.winmd for {ToolChain} toolchain", WindowsPlatform.GetCompilerName(Compiler));
 			}
 			string FoundationWinMDPath = HoloLensPlatform.GetLatestMetadataPathForApiContract("Windows.Foundation.FoundationContract", Compiler, Logger);
 			string UniversalWinMDPath = HoloLensPlatform.GetLatestMetadataPathForApiContract("Windows.Foundation.UniversalApiContract", Compiler, Logger);
diff --git a/Engine/Platforms/Hololens/Source/Programs/UnrealBuildTool/UEBuildHoloLens.cs b/Engine/Platforms/Hololens/Source/Programs/UnrealBuildTool/UEBuildHoloLens.cs
--- a/Engine/Platforms/Hololens/Source/Programs/UnrealBuildTool/UEBuildHoloLens.cs
+++ b/Engine/Platforms/Hololens/Source/Programs/UnrealBuildTool/UEBuildHoloLens.cs
@@ -440,7 +440,7 @@ namespace UnrealBuildTool
 			return base.GetBinaryExtension(InBinaryType);
 		}
 
-		internal static DirectoryReference? GetCppCXMetadataLocation(WindowsCompiler Compiler, string CompilerVersion, UnrealArch Architecture, ILogger Logger)
+		internal static DirectoryReference? GetCppCXMetadataLocation(WindowsCompiler Compiler, string? CompilerVersion, UnrealArch Architecture, ILogger Logger)
 		{
 			VersionNumber? SelectedToolChainVersion;
 			DirectoryReference? SelectedToolChainDir;
