UE-165131 HoloLens "Multi" packaging seems unnecessary, causes problems
-Revert to logging only if the toolchain is unavailable, and falling back to x64 throwing here breaks project generation.

diff --git a/Engine/Platforms/Hololens/Source/Programs/UnrealBuildTool/UEBuildHoloLens.cs b/Engine/Platforms/Hololens/Source/Programs/UnrealBuildTool/UEBuildHoloLens.cs
--- a/Engine/Platforms/Hololens/Source/Programs/UnrealBuildTool/UEBuildHoloLens.cs
+++ b/Engine/Platforms/Hololens/Source/Programs/UnrealBuildTool/UEBuildHoloLens.cs
@@ -245,35 +245,42 @@ namespace UnrealBuildTool
 			this.Logger = Logger;
 		}
 
+		static bool bShownLog = false;
+		static bool bShownLog2 = false;
 		public override UnrealArchitectures ActiveArchitectures(FileReference? ProjectFile, string? TargetName)
 		{
 			var ArchList = new List<UnrealArch>();
-			bool bBuildForEmulation;
+			bool bBuildForEmulation = false;
 			ConfigHierarchy Ini = ConfigCache.ReadHierarchy(ConfigHierarchyType.Engine, ProjectFile?.Directory, UnrealTargetPlatform.HoloLens);
-			if (Ini.GetBool("/Script/HoloLensPlatformEditor.HoloLensTargetSettings", "bBuildForEmulation", out bBuildForEmulation) && bBuildForEmulation)
+			Ini.GetBool("/Script/HoloLensPlatformEditor.HoloLensTargetSettings", "bBuildForEmulation", out bBuildForEmulation);
+			if (bBuildForEmulation)
 			{
 				bool bHasX64Toolchain = MicrosoftPlatformSDK.HasValidCompiler(GetCompiler(ProjectFile, UnrealArch.X64), UnrealArch.X64, Logger);
 				Log.TraceLogOnce($"HoloLens X64: Desired Compiler: {GetCompiler(ProjectFile, UnrealArch.X64)}, Valid: {bHasX64Toolchain}");
-				if (bHasX64Toolchain)
+				if (!bHasX64Toolchain && !bShownLog)
 				{
-					ArchList.Add(UnrealArch.X64);
-				}
-				else
-				{
-					throw new BuildException($"Project configured to build for HoloLens Emulation (on PC), but a valid X64 toolchain was not found!  Perhaps you need to install it?");
+					bShownLog = true;
+					Logger.LogInformation("X64 was requested for HoloLens, but the {0} X64 toolchain is not installed.", GetCompiler(ProjectFile, UnrealArch.X64));
 				}
+				
+				ArchList.Add(UnrealArch.X64);
 			} 
 			else
 			{
 				bool bHasArmToolchain = MicrosoftPlatformSDK.HasValidCompiler(GetCompiler(ProjectFile, UnrealArch.Arm64), UnrealArch.Arm64, Logger);
 				Log.TraceLogOnce($"HoloLens Arm64: Desired Compiler: {GetCompiler(ProjectFile, UnrealArch.Arm64)}, Valid: {bHasArmToolchain}");
-				if (bHasArmToolchain)
+				if (!bHasArmToolchain)
 				{
-					ArchList.Add(UnrealArch.Arm64);
+					if (!bShownLog2)
+					{
+						bShownLog2 = true;
+						Logger.LogInformation("Arm64 was requested for HoloLens, but the {0} Arm64 toolchain is not installed.  Falling back to X64, which will not run on device, because doing otherwise breaks project generation.  You may want to install the toolchain or you may want to avoid syncing the Hololens platform.", GetCompiler(ProjectFile, UnrealArch.Arm64));
+					}
+					ArchList.Add(UnrealArch.X64);
 				}
 				else
 				{
-					throw new BuildException($"Project configured to build for HoloLens device, but a valid ARM64 toolchain was not found!  Perhaps you need to install it?");
+					ArchList.Add(UnrealArch.Arm64);
 				}
 			}
 
