UE-170084 Constrain Hololens device detection infrastructure
-New ProjectSettings->Platforms->Hololens - Editor->bEditorAutomaticallyDetectsHoloLensDevices setting which defaults to false.  If true we scan for hololens devices during editor startup, when false we do not.  Either way we can manually connect to a hololens device.
-This had to be done as a new settings object because the platform extension target platform config file is read long after the device detection code has run.  It's also reasonable that this new setting be per-project, because you only want it to run on projects where you are actually using hololens and your device setup allows automatic device detection to work.

diff --git a/Engine/Platforms/Hololens/Source/Developer/HoloLensPlatformEditor/Private/HoloLensPlatformEditorModule.cpp b/Engine/Platforms/Hololens/Source/Developer/HoloLensPlatformEditor/Private/HoloLensPlatformEditorModule.cpp
--- a/Engine/Platforms/Hololens/Source/Developer/HoloLensPlatformEditor/Private/HoloLensPlatformEditorModule.cpp
+++ b/Engine/Platforms/Hololens/Source/Developer/HoloLensPlatformEditor/Private/HoloLensPlatformEditorModule.cpp
@@ -7,7 +7,7 @@
 #include "ISettingsModule.h"
 #include "Modules/ModuleManager.h"
 #include "HoloLensTargetSettings.h"
-
+#include "HoloLensEditorSettings.h"
 
 #define LOCTEXT_NAMESPACE "FHoloLensPlatformEditorModule"
 
@@ -43,6 +43,11 @@ public:
 				LOCTEXT("TargetSettingsDescription", "Settings for HoloLens"),
 				GetMutableDefault<UHoloLensTargetSettings>()
 			);
+			SettingsModule->RegisterSettings("Project", "Platforms", "HoloLensEditor",
+				LOCTEXT("EditorSettingsName", "HoloLens - Editor"),
+				LOCTEXT("EditorSettingsDescription", "Editor Settings for HoloLens"),
+				GetMutableDefault<UHoloLensEditorSettings>()
+			);
 		}
 
 		// register settings detail panel customization
@@ -70,6 +75,7 @@ public:
 		if (SettingsModule != nullptr)
 		{
 			SettingsModule->UnregisterSettings("Project", "Platforms", "HoloLens");
+			SettingsModule->UnregisterSettings("Project", "Platforms", "HoloLensEditor");
 		}
 	}
 
diff --git a/Engine/Platforms/Hololens/Source/Developer/HoloLensPlatformEditor/Private/HololensEditorSettings.cpp b/Engine/Platforms/Hololens/Source/Developer/HoloLensPlatformEditor/Private/HololensEditorSettings.cpp
new file mode 100644
--- /dev/null
+++ b/Engine/Platforms/Hololens/Source/Developer/HoloLensPlatformEditor/Private/HololensEditorSettings.cpp
@@ -0,0 +1,8 @@
+// Copyright Epic Games, Inc. All Rights Reserved.
+
+#include "HoloLensEditorSettings.h"
+
+UHoloLensEditorSettings::UHoloLensEditorSettings(const FObjectInitializer& ObjectInitializer)
+	: Super(ObjectInitializer)
+{
+}
\ No newline at end of file
diff --git a/Engine/Platforms/Hololens/Source/Developer/HoloLensPlatformEditor/Private/HololensEditorSettings.h b/Engine/Platforms/Hololens/Source/Developer/HoloLensPlatformEditor/Private/HololensEditorSettings.h
new file mode 100644
--- /dev/null
+++ b/Engine/Platforms/Hololens/Source/Developer/HoloLensPlatformEditor/Private/HololensEditorSettings.h
@@ -0,0 +1,29 @@
+// Copyright Epic Games, Inc. All Rights Reserved.
+
+/*=============================================================================
+	HoloLensEditorSettings.h: Declares the UHoloLensEditorSettings class.
+=============================================================================*/
+
+#pragma once
+
+#include "CoreMinimal.h"
+#include "UObject/Object.h"
+#include "HoloLensEditorSettings.generated.h"
+
+/**
+ * Implements the editor settings for the HoloLens target platform.
+ */
+UCLASS(config = EditorPerProjectUserSettings)
+class UHoloLensEditorSettings
+	: public UObject
+{
+public:
+	GENERATED_UCLASS_BODY()
+
+	/**
+	 * When true, we will scan for hololens devices during editor startup.  This does increase startup time.
+	 * One can also manually add hololens devices with Platforms->Device Manager->Add An Unlisted Device.
+	 */
+	UPROPERTY(EditAnywhere, Config, Category = "HoloLens")
+	bool bEditorAutomaticallyDetectsHoloLensDevices = false;
+};
diff --git a/Engine/Platforms/Hololens/Source/Developer/HoloLensTargetPlatform/Private/HoloLensTargetPlatform.cpp b/Engine/Platforms/Hololens/Source/Developer/HoloLensTargetPlatform/Private/HoloLensTargetPlatform.cpp
--- a/Engine/Platforms/Hololens/Source/Developer/HoloLensTargetPlatform/Private/HoloLensTargetPlatform.cpp
+++ b/Engine/Platforms/Hololens/Source/Developer/HoloLensTargetPlatform/Private/HoloLensTargetPlatform.cpp
@@ -96,8 +96,8 @@ FHoloLensTargetPlatform::FHoloLensTargetPlatform(bool bIsClientOnly)
 	DeviceDetectedRegistration = DeviceDetector->OnDeviceDetected().AddRaw(this, &FHoloLensTargetPlatform::OnDeviceDetected);
 
 #if WITH_EDITOR
-	// Don't automatically try to detect a device if running editor with -server.
-	if (!IsRunningDedicatedServer() && !IsRunningCommandlet())
+	GConfig->GetBool(TEXT("/Script/HoloLensPlatformEditor.HoloLensEditorSettings"), TEXT("bEditorAutomaticallyDetectsHoloLensDevices"), bEditorAutomaticallyDetectsHoloLensDevices, GEditorPerProjectIni);
+	if (bEditorAutomaticallyDetectsHoloLensDevices && !IsRunningDedicatedServer() && !IsRunningCommandlet() )
 #endif
 	{
 		DeviceDetector->StartDeviceDetection();
@@ -111,6 +111,13 @@ FHoloLensTargetPlatform::~FHoloLensTargetPlatform()
 
 void FHoloLensTargetPlatform::GetAllDevices(TArray<ITargetDevicePtr>& OutDevices) const
 {
+#if WITH_EDITOR
+	if (!bEditorAutomaticallyDetectsHoloLensDevices)
+	{
+		return;
+	}
+#endif
+
 	DeviceDetector->StartDeviceDetection();
 
 	OutDevices.Reset();
diff --git a/Engine/Platforms/Hololens/Source/Developer/HoloLensTargetPlatform/Private/HoloLensTargetPlatform.h b/Engine/Platforms/Hololens/Source/Developer/HoloLensTargetPlatform/Private/HoloLensTargetPlatform.h
--- a/Engine/Platforms/Hololens/Source/Developer/HoloLensTargetPlatform/Private/HoloLensTargetPlatform.h
+++ b/Engine/Platforms/Hololens/Source/Developer/HoloLensTargetPlatform/Private/HoloLensTargetPlatform.h
@@ -115,6 +115,9 @@ private:
 	FStaticMeshLODSettings StaticMeshLODSettings;
 #endif //WITH_ENGINE
 
+#if WITH_EDITOR
+	bool bEditorAutomaticallyDetectsHoloLensDevices = false;
+#endif
 };
 
 
