UE-167282 GitHub 9664 : Moved HoloLens target platform definition to platform extension code
-Adding the code for the move as a reference for future platforms but commented out because there are third party plugin concerns we do not want to spend time on now.  The comment explains further.

diff --git a/Engine/Platforms/Hololens/Source/Programs/UnrealBuildTool/UEBuildHoloLens.cs b/Engine/Platforms/Hololens/Source/Programs/UnrealBuildTool/UEBuildHoloLens.cs
--- a/Engine/Platforms/Hololens/Source/Programs/UnrealBuildTool/UEBuildHoloLens.cs
+++ b/Engine/Platforms/Hololens/Source/Programs/UnrealBuildTool/UEBuildHoloLens.cs
@@ -24,6 +24,21 @@ namespace UnrealBuildTool
 		public static UnrealPlatformGroup HoloLens = FindOrAddByName("HoloLens");
 	}
 
+	// This is the correct way to define a new UnrealTargetPlatform for a platform extension.
+	// We are not currently doing this for HoloLens out of concern that a third party plugin might have trouble accessing it.  However, this concern is hypothetical.  Just haven't spent the time to test it.
+	// In the long run even if that concern does turn out to be founded it would be better to create some mechanism to handle it which preserves the encapsulation of the platform.  Perhaps third party plugins would
+	// need to have their own plugin extension files (ie ThePlugin_HoloLens.uplugin) it would "just work"?
+	// NDA platforms *Must* do it this way, to avoid leaking NDA information outside the platform extension.  It's not so strictly required on public platform extension platforms 
+	// like HoloLens, but it would still be nice to maintain the encapsulation.
+	// See UEBuildTarget.cs "UnrealTargetPlatform HoloLens" for the small amount of code that would, ideally, be removed.
+	//partial struct UnrealTargetPlatform
+	//{
+	//	/// <summary>
+	//	/// HoloLens
+	//	/// </summary>
+	//	public static UnrealTargetPlatform HoloLens = FindOrAddByName("HoloLens");
+	//}
+
 	/// <summary>
 	/// HoloLens-specific target settings
 	/// </summary>
