More AutomationTool structured logging changes.

diff --git a/Engine/Platforms/Hololens/Source/Programs/AutomationTool/Gauntlet/Gauntlet.TargetDeviceHoloLens.cs b/Engine/Platforms/Hololens/Source/Programs/AutomationTool/Gauntlet/Gauntlet.TargetDeviceHoloLens.cs
--- a/Engine/Platforms/Hololens/Source/Programs/AutomationTool/Gauntlet/Gauntlet.TargetDeviceHoloLens.cs
+++ b/Engine/Platforms/Hololens/Source/Programs/AutomationTool/Gauntlet/Gauntlet.TargetDeviceHoloLens.cs
@@ -14,6 +14,8 @@ using UnrealBuildBase;
 using EpicGames.Core;
 using Microsoft.Extensions.Logging;
 
+using static AutomationTool.CommandUtils;
+
 namespace Gauntlet
 {
 	[System.Runtime.Versioning.SupportedOSPlatform("windows")]
@@ -279,25 +281,25 @@ namespace Gauntlet
 				ExternallyLaunchedProcess = null;
 				if (!ProcessToKill.CloseMainWindow())
 				{
-					CommandUtils.LogWarning("{0} did not respond to close request.  Killing...", ProcessNameForLogging);
+					Logger.LogWarning("{ProcessNameForLogging} did not respond to close request.  Killing...", ProcessNameForLogging);
 					ProcessToKill.Kill();
 					ProcessToKill.WaitForExit(60000);
 				}
 				if (!ProcessToKill.HasExited)
 				{
-					CommandUtils.LogLog("Process {0} failed to exit.", ProcessNameForLogging);
+					Logger.LogDebug("Process {ProcessNameForLogging} failed to exit.", ProcessNameForLogging);
 				}
 				else
 				{
-					CommandUtils.LogLog("Process {0} successfully exited.", ProcessNameForLogging);
+					Logger.LogDebug("Process {ProcessNameForLogging} successfully exited.", ProcessNameForLogging);
 					OnProcessExited();
 				}
 				ProcessToKill.Close();
 			}
 			catch (Exception Ex)
 			{
-				CommandUtils.LogWarning("Exception while trying to kill process {0}:", ProcessNameForLogging);
-				CommandUtils.LogWarning(LogUtils.FormatException(Ex));
+				Logger.LogWarning("Exception while trying to kill process {ProcessNameForLogging}:", ProcessNameForLogging);
+				Logger.LogWarning("{Text}", LogUtils.FormatException(Ex));
 			}
 		}
 
@@ -1336,4 +1338,4 @@ namespace Gauntlet
  
  
  
- 
\ No newline at end of file
+ 
diff --git a/Engine/Platforms/Hololens/Source/Programs/AutomationTool/HoloLensPlatform.Automation.cs b/Engine/Platforms/Hololens/Source/Programs/AutomationTool/HoloLensPlatform.Automation.cs
--- a/Engine/Platforms/Hololens/Source/Programs/AutomationTool/HoloLensPlatform.Automation.cs
+++ b/Engine/Platforms/Hololens/Source/Programs/AutomationTool/HoloLensPlatform.Automation.cs
@@ -22,6 +22,8 @@ using AutomationScripts;
 using System.Runtime.Versioning;
 using Microsoft.Extensions.Logging;
 
+using static AutomationTool.CommandUtils;
+
 namespace HoloLens.Automation
 {
 	class HoloLensLauncherCreatedProcess : IProcessResult
@@ -136,25 +138,25 @@ namespace HoloLens.Automation
 				ExternallyLaunchedProcess = null;
 				if (!ProcessToKill.CloseMainWindow())
 				{
-					CommandUtils.LogWarning("{0} did not respond to close request.  Killing...", ProcessNameForLogging);
+					Logger.LogWarning("{ProcessNameForLogging} did not respond to close request.  Killing...", ProcessNameForLogging);
 					ProcessToKill.Kill();
 					ProcessToKill.WaitForExit(60000);
 				}
 				if (!ProcessToKill.HasExited)
 				{
-					CommandUtils.LogLog("Process {0} failed to exit.", ProcessNameForLogging);
+					Logger.LogDebug("Process {ProcessNameForLogging} failed to exit.", ProcessNameForLogging);
 				}
 				else
 				{
-					CommandUtils.LogLog("Process {0} successfully exited.", ProcessNameForLogging);
+					Logger.LogDebug("Process {ProcessNameForLogging} successfully exited.", ProcessNameForLogging);
 					OnProcessExited();
 				}
 				ProcessToKill.Close();
 			}
 			catch (Exception Ex)
 			{
-				CommandUtils.LogWarning("Exception while trying to kill process {0}:", ProcessNameForLogging);
-				CommandUtils.LogWarning(LogUtils.FormatException(Ex));
+				Logger.LogWarning("Exception while trying to kill process {ProcessNameForLogging}:", ProcessNameForLogging);
+				Logger.LogWarning("{Text}", LogUtils.FormatException(Ex));
 			}
 		}
 
diff --git a/Engine/Source/Programs/AutomationTool/AutomationUtils/FileRegions.cs b/Engine/Source/Programs/AutomationTool/AutomationUtils/FileRegions.cs
--- a/Engine/Source/Programs/AutomationTool/AutomationUtils/FileRegions.cs
+++ b/Engine/Source/Programs/AutomationTool/AutomationUtils/FileRegions.cs
@@ -81,12 +81,12 @@ namespace AutomationTool
 					Suffix++;
 				}
 
-				CommandUtils.LogInformation(" | {0,-17} | {1,18:#,##} ({2,7:0.00} {3}) | {4,13:#,##} |",
+				Logger.LogInformation("{Message}", String.Format(" | {0,-17} | {1,18:#,##} ({2,7:0.00} {3}) | {4,13:#,##} |",
 					Pair.Key.ToString().ToUpperInvariant(),
 					TotalSizeInBytes,
 					TotalSizeInUnits,
 					Suffixes[Suffix],
-					Pair.Value.Count);
+					Pair.Value.Count));
 			}
 			Logger.LogInformation(" +-------------------+----------------------------------+---------------+");
 		}
diff --git a/Engine/Source/Programs/AutomationTool/AutomationUtils/MCPPublic.cs b/Engine/Source/Programs/AutomationTool/AutomationUtils/MCPPublic.cs
--- a/Engine/Source/Programs/AutomationTool/AutomationUtils/MCPPublic.cs
+++ b/Engine/Source/Programs/AutomationTool/AutomationUtils/MCPPublic.cs
@@ -2562,12 +2562,10 @@ namespace EpicGames.MCP.Config
                                	var Inner = Ex.InnerException;
 								while(null != Inner)
 								{
-									BuildCommand.LogWarning("Exception encountered creating McpConfig [{0}] with error: {1}",
-										PotentialConfigType.Name, Inner.Message);
+									Logger.LogWarning("Exception encountered creating McpConfig [{Name}] with error: {Message}", PotentialConfigType.Name, Inner.Message);
 									Inner = Inner.InnerException;
 								}
-								BuildCommand.LogWarning("Unable to create McpConfig [{0}] with error: {1} \n {2}",
-									PotentialConfigType.Name, Ex.Message, Ex.StackTrace);
+								Logger.LogWarning("Unable to create McpConfig [{Name}] with error: {Message} \n {Trace}", PotentialConfigType.Name, Ex.Message, Ex.StackTrace);
                             }
                         }
                     }
diff --git a/Engine/Source/Programs/AutomationTool/AutomationUtils/ProcessUtils.cs b/Engine/Source/Programs/AutomationTool/AutomationUtils/ProcessUtils.cs
--- a/Engine/Source/Programs/AutomationTool/AutomationUtils/ProcessUtils.cs
+++ b/Engine/Source/Programs/AutomationTool/AutomationUtils/ProcessUtils.cs
@@ -720,7 +720,7 @@ namespace AutomationTool
 			{
 				foreach (var Item in ExeToTimeInMs)
 				{
-					LogVerbose("Total {0}s to run " + Item.Key, Item.Value / 1000);
+					Logger.LogDebug("Total {Time}s to run {Exe}", Item.Value / 1000, Item.Key);
 				}
 				ExeToTimeInMs.Clear();
 			}
diff --git a/Engine/Source/Programs/AutomationTool/Scripts/RunProjectCommand.Automation.cs b/Engine/Source/Programs/AutomationTool/Scripts/RunProjectCommand.Automation.cs
--- a/Engine/Source/Programs/AutomationTool/Scripts/RunProjectCommand.Automation.cs
+++ b/Engine/Source/Programs/AutomationTool/Scripts/RunProjectCommand.Automation.cs
@@ -208,8 +208,7 @@ namespace AutomationScripts
 							Logger.LogInformation("Stopping cook server...");
 							CookServerProcess.StopProcess();
 						}
-						LogInformation("Cook server exited with error code: {0} (see {1} for more info)",
-								CookServerProcess.ExitCode, CookServerLogFile);
+						Logger.LogInformation("Cook server exited with error code: {ExitCode} (see {File} for more info)", CookServerProcess.ExitCode, CookServerLogFile);
 					}
 					if (DedicatedServerProcess != null)
 					{
@@ -218,8 +217,7 @@ namespace AutomationScripts
 							Logger.LogInformation("Stopping dedicated server...");
 							DedicatedServerProcess.StopProcess();
 						}
-						LogInformation("Dedicated server exited with error code: {0} (see {1} for more info)",
-								DedicatedServerProcess.ExitCode, DedicatedServerLogFile);
+						Logger.LogInformation("Dedicated server exited with error code: {ExitCode} (see {File} for more info)", DedicatedServerProcess.ExitCode, DedicatedServerLogFile);
 					}
 				}
 				CopyLogsBackToLogFolder();
@@ -377,8 +375,7 @@ namespace AutomationScripts
 								Logger.LogInformation("Welcomed by server or client loaded");
 							}
 							Logger.LogInformation("Test complete");
-							LogInformation("**** UNATTENDED TEST COMPLETE: {0:0.00} seconds ****",
-								(DateTime.UtcNow - ClientStartTime).TotalMilliseconds / 1000);
+							Logger.LogInformation("**** UNATTENDED TEST COMPLETE: {Time} seconds ****", $"{(DateTime.UtcNow - ClientStartTime).TotalMilliseconds / 1000:0.00}");
 							bTestExitTextFound = true;
 							bKeepReading = false;
 						}
