Updated the platform to create unique config names because the config/platform combination conflicts with other platforms(win64).

diff --git a/Engine/Platforms/Hololens/Source/Programs/UnrealBuildTool/HoloLensProjectGenerator.cs b/Engine/Platforms/Hololens/Source/Programs/UnrealBuildTool/HoloLensProjectGenerator.cs
--- a/Engine/Platforms/Hololens/Source/Programs/UnrealBuildTool/HoloLensProjectGenerator.cs
+++ b/Engine/Platforms/Hololens/Source/Programs/UnrealBuildTool/HoloLensProjectGenerator.cs
@@ -168,6 +168,12 @@ namespace UnrealBuildTool
 			return InPlatform.ToString();
 		}
 
+		public override bool RequiresDistinctVisualStudioConfigurationName(UnrealTargetPlatform InPlatform, UnrealTargetConfiguration InConfiguration, DirectoryReference InProjectDir)
+		{
+			// Win64 and HoloLens support the same architectures so the config names need to be unique.
+			return true;
+		}
+
 		private bool IsValidHoloLensTarget(UnrealTargetPlatform InPlatform, TargetType InTargetType, FileReference InTargetFilePath)
 		{
 			if ((InPlatform == UnrealTargetPlatform.HoloLens) &&
