Fix Hololens build.

diff --git a/Engine/Platforms/Hololens/Source/Runtime/D3D12RHI/Public/D3D12ThirdParty.h b/Engine/Platforms/Hololens/Source/Runtime/D3D12RHI/Public/D3D12ThirdParty.h
--- a/Engine/Platforms/Hololens/Source/Runtime/D3D12RHI/Public/D3D12ThirdParty.h
+++ b/Engine/Platforms/Hololens/Source/Runtime/D3D12RHI/Public/D3D12ThirdParty.h
@@ -6,6 +6,8 @@
 #pragma warning(push)
 #pragma warning(disable : 4005)	// macro redefinition
 
+#include "HAL/Platform.h"
+
 #if PLATFORM_64BITS
 	#pragma pack(push,16)
 #else
@@ -14,6 +16,8 @@
 
 #include "Microsoft/AllowMicrosoftPlatformTypes.h"
 THIRD_PARTY_INCLUDES_START
+	#define INITGUID  // for CLSID_D3D12DeviceRemovedExtendedData
+	#include <guiddef.h>
 	#include <d3d12.h>
 	#include <d3dx12.h>
 	#include <dxgi1_6.h>
