HoloLens: Disable warning 4946, reinterpret cast used between related types
-This is disabled for Windows, Hololens has no particular reason to be more strict.

diff --git a/Engine/Platforms/Hololens/Source/Runtime/Core/Public/HoloLensPlatformCompilerSetup.h b/Engine/Platforms/Hololens/Source/Runtime/Core/Public/HoloLensPlatformCompilerSetup.h
--- a/Engine/Platforms/Hololens/Source/Runtime/Core/Public/HoloLensPlatformCompilerSetup.h
+++ b/Engine/Platforms/Hololens/Source/Runtime/Core/Public/HoloLensPlatformCompilerSetup.h
@@ -174,7 +174,7 @@ static_assert(_MSC_VER >= 1929, "Visual Studio 2019 v16.11 or greater is require
 #pragma warning(default : 4906) // string literal cast to 'LPWSTR'
 #pragma warning(disable : 4917) // 'declarator' : a GUID cannot only be associated with a class, interface or namespace ( ocid.h breaks this)
 #pragma warning(default : 4931) // we are assuming the type library was built for number-bit pointers
-#pragma warning(default : 4946) // reinterpret_cast used between related classes: 'class1' and 'class2'
+#pragma warning(disable : 4946) // reinterpret_cast used between related classes: 'class1' and 'class2'
 #pragma warning(default : 4928) // illegal copy-initialization; more than one user-defined conversion has been implicitly applied
 #pragma warning(disable : 4180) // qualifier applied to function type has no meaning; ignored
 #pragma warning( disable : 4121 )
