Moving more D3D12 defines
- D3D12_USE_DUMMY_BACKBUFFER from Public to Private.
- FD3D12_TEXTURE_DATA_PITCH_ALIGNMENT into the platform abstractions.
- FD3D12_DESCRIPTOR_HEAP_TYPE_SAMPLER from Public to Private.

diff --git a/Engine/Platforms/Hololens/Source/Runtime/D3D12RHI/Private/D3D12RHIBasePrivate.h b/Engine/Platforms/Hololens/Source/Runtime/D3D12RHI/Private/D3D12RHIBasePrivate.h
--- a/Engine/Platforms/Hololens/Source/Runtime/D3D12RHI/Private/D3D12RHIBasePrivate.h
+++ b/Engine/Platforms/Hololens/Source/Runtime/D3D12RHI/Private/D3D12RHIBasePrivate.h
@@ -6,6 +6,8 @@
 
 #pragma once
 
+#include "D3D12ThirdParty.h"
+
 // Standard Hololens implementation, used to mark the type as 'final'.
 struct FD3D12Payload final : public FD3D12PayloadBase
 {
@@ -16,13 +18,14 @@ struct FD3D12Payload final : public FD3D12PayloadBase
 
 // Hololens uses the graphics command list interface for the copy queue.
 typedef ID3D12GraphicsCommandList ID3D12CopyCommandList;
+
 #define D3D12RHI_PLATFORM_COPY_COMMAND_LIST_TYPE D3D12_COMMAND_LIST_TYPE_COPY
 
+#define FD3D12_DESCRIPTOR_HEAP_TYPE_SAMPLER D3D12_DESCRIPTOR_HEAP_TYPE_SAMPLER
+
 // Hololens uses standard D3D12 query heaps for timestamps
 #define D3D12RHI_PLATFORM_USES_TIMESTAMP_QUERIES 1
 
-#include "D3D12ThirdParty.h"
-
 #define D3D12RHI_RESOURCE_FLAG_ALLOW_INDIRECT_BUFFER D3D12_RESOURCE_FLAG_NONE
 #define D3D12RHI_HEAP_FLAG_ALLOW_INDIRECT_BUFFERS		D3D12_HEAP_FLAG_NONE
 
@@ -34,6 +37,7 @@ typedef ID3D12GraphicsCommandList ID3D12CopyCommandList;
 #define FD3D12_HEAP_FLAG_CREATE_NOT_ZEROED D3D12_HEAP_FLAG_NONE
 
 #define USE_STATIC_ROOT_SIGNATURE			0
+#define D3D12_USE_DUMMY_BACKBUFFER			1
 
 #include "D3D12Util.h"
 
diff --git a/Engine/Source/Runtime/D3D12RHI/Private/Windows/D3D12RHIBasePrivate.h b/Engine/Source/Runtime/D3D12RHI/Private/Windows/D3D12RHIBasePrivate.h
--- a/Engine/Source/Runtime/D3D12RHI/Private/Windows/D3D12RHIBasePrivate.h
+++ b/Engine/Source/Runtime/D3D12RHI/Private/Windows/D3D12RHIBasePrivate.h
@@ -6,6 +6,8 @@
 
 #pragma once
 
+#include "Windows/D3D12ThirdParty.h"
+
 // Standard Windows implementation, used to mark the type as 'final'.
 struct FD3D12Payload final : public FD3D12PayloadBase
 {
@@ -16,13 +18,14 @@ struct FD3D12Payload final : public FD3D12PayloadBase
 
 // Windows uses the graphics command list interface for the copy queue.
 typedef ID3D12GraphicsCommandList ID3D12CopyCommandList;
+
 #define D3D12RHI_PLATFORM_COPY_COMMAND_LIST_TYPE D3D12_COMMAND_LIST_TYPE_COPY
 
+#define FD3D12_DESCRIPTOR_HEAP_TYPE_SAMPLER D3D12_DESCRIPTOR_HEAP_TYPE_SAMPLER
+
 // Windows uses standard D3D12 query heaps for timestamps
 #define D3D12RHI_PLATFORM_USES_TIMESTAMP_QUERIES 1
 
-#include "Windows/D3D12ThirdParty.h"
-
 #define D3D12RHI_RESOURCE_FLAG_ALLOW_INDIRECT_BUFFER D3D12_RESOURCE_FLAG_NONE
 #define D3D12RHI_HEAP_FLAG_ALLOW_INDIRECT_BUFFERS		D3D12_HEAP_FLAG_NONE
 
@@ -30,6 +33,7 @@ typedef ID3D12GraphicsCommandList ID3D12CopyCommandList;
 #define D3D12RHI_NEEDS_SHADER_FEATURE_CHECKS 1
 
 #define USE_STATIC_ROOT_SIGNATURE			0
+#define D3D12_USE_DUMMY_BACKBUFFER			1
 
 #include "D3D12Util.h"
 #include "WindowsD3D12DiskCache.h"
diff --git a/Engine/Source/Runtime/D3D12RHI/Public/D3D12RHI.h b/Engine/Source/Runtime/D3D12RHI/Public/D3D12RHI.h
--- a/Engine/Source/Runtime/D3D12RHI/Public/D3D12RHI.h
+++ b/Engine/Source/Runtime/D3D12RHI/Public/D3D12RHI.h
@@ -43,19 +43,12 @@
 	#define PIPELINE_STATE_FILE_LOCATION			FPaths::ProjectSavedDir()
 	#define USE_PIX									WITH_PIX_EVENT_RUNTIME
 	#define D3D12RHI_SUPPORTS_WIN_PIX				USE_PIX
-#else
-	#include "D3D12RHIPlatformPublic.h"
-#endif
 
-#if PLATFORM_WINDOWS || PLATFORM_HOLOLENS
-#define D3D12_USE_DUMMY_BACKBUFFER 1
+	#define FD3D12_TEXTURE_DATA_PITCH_ALIGNMENT		D3D12_TEXTURE_DATA_PITCH_ALIGNMENT
 #else
-#define D3D12_USE_DUMMY_BACKBUFFER 0
+	#include "D3D12RHIPlatformPublic.h"
 #endif
 
-#define FD3D12_TEXTURE_DATA_PITCH_ALIGNMENT D3D12_TEXTURE_DATA_PITCH_ALIGNMENT
-#define FD3D12_DESCRIPTOR_HEAP_TYPE_SAMPLER D3D12_DESCRIPTOR_HEAP_TYPE_SAMPLER
-
 typedef uint16 CBVSlotMask;
 static_assert(MAX_ROOT_CBVS <= MAX_CBS, "MAX_ROOT_CBVS must be <= MAX_CBS.");
 static_assert((8 * sizeof(CBVSlotMask)) >= MAX_CBS, "CBVSlotMask isn't large enough to cover all CBs. Please increase the size.");
