Reorganizing D3D12ThirdParty headers to use COMPILED_PLATFORM_HEADER.

diff --git a/Engine/Platforms/Hololens/Source/Runtime/D3D12RHI/Private/D3D12RHIBasePrivate.h b/Engine/Platforms/Hololens/Source/Runtime/D3D12RHI/Private/D3D12RHIBasePrivate.h
--- a/Engine/Platforms/Hololens/Source/Runtime/D3D12RHI/Private/D3D12RHIBasePrivate.h
+++ b/Engine/Platforms/Hololens/Source/Runtime/D3D12RHI/Private/D3D12RHIBasePrivate.h
@@ -6,8 +6,6 @@
 
 #pragma once
 
-#include "D3D12ThirdParty.h"
-
 // Hololens uses the graphics command list interface for the copy queue.
 typedef ID3D12GraphicsCommandList ID3D12CopyCommandList;
 
diff --git a/Engine/Platforms/Hololens/Source/Runtime/D3D12RHI/Public/D3D12ThirdParty.h b/Engine/Platforms/Hololens/Source/Runtime/D3D12RHI/Public/D3D12ThirdParty.h
deleted file mode 100644
diff --git a/Engine/Platforms/Hololens/Source/Runtime/D3D12RHI/Public/HoloLensD3D12ThirdParty.h b/Engine/Platforms/Hololens/Source/Runtime/D3D12RHI/Public/HoloLensD3D12ThirdParty.h
new file mode 100644
--- /dev/null
+++ b/Engine/Platforms/Hololens/Source/Runtime/D3D12RHI/Public/HoloLensD3D12ThirdParty.h
@@ -0,0 +1,31 @@
+// Copyright Epic Games, Inc. All Rights Reserved.
+
+#pragma once
+
+// Disable macro redefinition warning for compatibility with Windows SDK 8+
+#pragma warning(push)
+#pragma warning(disable : 4005)	// macro redefinition
+
+#include "HAL/Platform.h"
+
+#if PLATFORM_64BITS
+	#pragma pack(push,16)
+#else
+	#pragma pack(push,8)
+#endif
+
+#include "Microsoft/AllowMicrosoftPlatformTypes.h"
+THIRD_PARTY_INCLUDES_START
+	#define INITGUID  // for CLSID_D3D12DeviceRemovedExtendedData
+	#include <guiddef.h>
+	#include <d3d12.h>
+	#include <d3dx12.h>
+	#include <dxgi1_6.h>
+	#include <dxgidebug.h>
+THIRD_PARTY_INCLUDES_END
+#include "Microsoft/HideMicrosoftPlatformTypes.h"
+
+#undef DrawText
+
+#pragma pack(pop)
+#pragma warning(pop)
diff --git a/Engine/Plugins/VirtualProduction/TextureShare/Source/TextureShareCore/Private/Core/ResourcesCache/D3D12/TextureShareCoreD3D12ResourcesCacheItem.cpp b/Engine/Plugins/VirtualProduction/TextureShare/Source/TextureShareCore/Private/Core/ResourcesCache/D3D12/TextureShareCoreD3D12ResourcesCacheItem.cpp
--- a/Engine/Plugins/VirtualProduction/TextureShare/Source/TextureShareCore/Private/Core/ResourcesCache/D3D12/TextureShareCoreD3D12ResourcesCacheItem.cpp
+++ b/Engine/Plugins/VirtualProduction/TextureShare/Source/TextureShareCore/Private/Core/ResourcesCache/D3D12/TextureShareCoreD3D12ResourcesCacheItem.cpp
@@ -4,7 +4,7 @@
 #include "Module/TextureShareCoreLog.h"
 #include "Templates/RefCounting.h"
 
-#include "Windows/D3D12ThirdParty.h"
+#include COMPILED_PLATFORM_HEADER(D3D12ThirdParty.h)
 
 namespace UE::TextureShareCore::D3D12ResourcesCacheItem
 {
diff --git a/Engine/Source/Runtime/D3D12RHI/Private/D3D12Descriptors.h b/Engine/Source/Runtime/D3D12RHI/Private/D3D12Descriptors.h
--- a/Engine/Source/Runtime/D3D12RHI/Private/D3D12Descriptors.h
+++ b/Engine/Source/Runtime/D3D12RHI/Private/D3D12Descriptors.h
@@ -10,12 +10,6 @@
 #include "Containers/Queue.h"
 #include "Containers/List.h"
 
-#if PLATFORM_WINDOWS
-#include "Windows/D3D12ThirdParty.h"
-#else
-#include "D3D12ThirdParty.h"
-#endif
-
 inline D3D12_DESCRIPTOR_HEAP_TYPE Translate(ERHIDescriptorHeapType InHeapType)
 {
 	switch (InHeapType)
diff --git a/Engine/Source/Runtime/D3D12RHI/Private/D3D12RHICommon.h b/Engine/Source/Runtime/D3D12RHI/Private/D3D12RHICommon.h
--- a/Engine/Source/Runtime/D3D12RHI/Private/D3D12RHICommon.h
+++ b/Engine/Source/Runtime/D3D12RHI/Private/D3D12RHICommon.h
@@ -8,11 +8,7 @@
 
 #include "HAL/Platform.h"
 
-#if PLATFORM_WINDOWS
-	#include "Windows/D3D12ThirdParty.h"
-#else
-	#include "D3D12ThirdParty.h"
-#endif
+#include COMPILED_PLATFORM_HEADER(D3D12ThirdParty.h)
 
 #include "D3D12RHI.h"
 #include "D3D12RHIDefinitions.h"
diff --git a/Engine/Source/Runtime/D3D12RHI/Private/Windows/D3D12RHIBasePrivate.h b/Engine/Source/Runtime/D3D12RHI/Private/Windows/D3D12RHIBasePrivate.h
--- a/Engine/Source/Runtime/D3D12RHI/Private/Windows/D3D12RHIBasePrivate.h
+++ b/Engine/Source/Runtime/D3D12RHI/Private/Windows/D3D12RHIBasePrivate.h
@@ -6,8 +6,6 @@
 
 #pragma once
 
-#include "Windows/D3D12ThirdParty.h"
-
 // Windows uses the graphics command list interface for the copy queue.
 typedef ID3D12GraphicsCommandList ID3D12CopyCommandList;
 
diff --git a/Engine/Source/Runtime/D3D12RHI/Public/ID3D12DynamicRHI.h b/Engine/Source/Runtime/D3D12RHI/Public/ID3D12DynamicRHI.h
--- a/Engine/Source/Runtime/D3D12RHI/Public/ID3D12DynamicRHI.h
+++ b/Engine/Source/Runtime/D3D12RHI/Public/ID3D12DynamicRHI.h
@@ -5,11 +5,7 @@
 #include "RHI.h"
 #include "Templates/Function.h"
 
-#if PLATFORM_WINDOWS
-	#include "Windows/D3D12ThirdParty.h"
-#else
-	#include "D3D12ThirdParty.h"
-#endif
+#include COMPILED_PLATFORM_HEADER(D3D12ThirdParty.h)
 
 struct FD3D12MinimalAdapterDesc
 {
diff --git a/Engine/Source/Runtime/D3D12RHI/Public/Windows/D3D12ThirdParty.h b/Engine/Source/Runtime/D3D12RHI/Public/Windows/D3D12ThirdParty.h
deleted file mode 100644
diff --git a/Engine/Source/Runtime/D3D12RHI/Public/Windows/WindowsD3D12ThirdParty.h b/Engine/Source/Runtime/D3D12RHI/Public/Windows/WindowsD3D12ThirdParty.h
new file mode 100644
--- /dev/null
+++ b/Engine/Source/Runtime/D3D12RHI/Public/Windows/WindowsD3D12ThirdParty.h
@@ -0,0 +1,30 @@
+// Copyright Epic Games, Inc. All Rights Reserved.
+
+#pragma once
+
+// Disable macro redefinition warning for compatibility with Windows SDK 8+
+#pragma warning(push)
+#pragma warning(disable : 4005)	// macro redefinition
+
+#include "HAL/Platform.h"
+
+#if PLATFORM_64BITS
+	#pragma pack(push,16)
+#else
+	#pragma pack(push,8)
+#endif
+
+#include "Microsoft/AllowMicrosoftPlatformTypes.h"
+#define INITGUID
+THIRD_PARTY_INCLUDES_START
+	#include <d3d12.h>
+	#include <d3dx12.h>
+	#include <dxgi1_6.h>
+	#include <dxgidebug.h>
+THIRD_PARTY_INCLUDES_END
+#include "Microsoft/HideMicrosoftPlatformTypes.h"
+
+#undef DrawText
+
+#pragma pack(pop)
+#pragma warning(pop)
