Add speculative fix to HoloLens crash by gating adding default resources to natural resources on default resources being valid when generating child content, rather than doing a check().

diff --git a/Engine/Platforms/Hololens/Source/Developer/HoloLensPlatformEditor/Private/HoloLensLocalizedResourcesCustomization.cpp b/Engine/Platforms/Hololens/Source/Developer/HoloLensPlatformEditor/Private/HoloLensLocalizedResourcesCustomization.cpp
--- a/Engine/Platforms/Hololens/Source/Developer/HoloLensPlatformEditor/Private/HoloLensLocalizedResourcesCustomization.cpp
+++ b/Engine/Platforms/Hololens/Source/Developer/HoloLensPlatformEditor/Private/HoloLensLocalizedResourcesCustomization.cpp
@@ -158,10 +158,11 @@ void FHoloLensLocalizedResourcesNodeBuilder::GenerateChildContent(IDetailChildre
 		}
 	}
 
-	check(DefaultResources.IsValid());
-
-	IDetailGroup& NeutralGroup = ChildrenBuilder.AddGroup("NeutralResources", LOCTEXT("NeutralResources", "Neutral Resources"));
-	NeutralGroup.AddPropertyRow(DefaultResources.ToSharedRef());
+	if (DefaultResources.IsValid())
+	{
+		IDetailGroup& NeutralGroup = ChildrenBuilder.AddGroup("NeutralResources", LOCTEXT("NeutralResources", "Neutral Resources"));
+		NeutralGroup.AddPropertyRow(DefaultResources.ToSharedRef());
+	}
 
 	IDetailGroup& LocalizedGroup = ChildrenBuilder.AddGroup("LocalizedResources", LOCTEXT("LocalizedResources", "Localized Resources"));
 
